Sequencer Scripting: Move to "Scripting" category alongside Python Script Plugin

diff --git a/Engine/Plugins/MovieScene/SequencerScripting/SequencerScripting.uplugin b/Engine/Plugins/MovieScene/SequencerScripting/SequencerScripting.uplugin
--- a/Engine/Plugins/MovieScene/SequencerScripting/SequencerScripting.uplugin
+++ b/Engine/Plugins/MovieScene/SequencerScripting/SequencerScripting.uplugin
@@ -2,16 +2,16 @@
 	"FileVersion": 3,
 	"Version": 1,
 	"VersionName": "1.0",
-	"FriendlyName": "SequencerScripting",
+	"FriendlyName": "Sequencer Scripting",
 	"Description": "Python and blutility scripting extensions for sequencer and movie scenes",
-	"Category": "Other",
+	"Category": "Scripting",
 	"CreatedBy": "Epic Games Inc.",
 	"CreatedByURL": "",
 	"DocsURL": "",
 	"MarketplaceURL": "",
 	"SupportURL": "",
 	"CanContainContent": true,
-	"IsBetaVersion": false,
+	"IsBetaVersion": true,
 	"Installed": false,
 	"Plugins": [
 		{
