Sequencer: Add get/set view and working ranges for python

diff --git a/Engine/Plugins/MovieScene/SequencerScripting/Source/SequencerScripting/Private/ExtensionLibraries/MovieSceneSequenceExtensions.cpp b/Engine/Plugins/MovieScene/SequencerScripting/Source/SequencerScripting/Private/ExtensionLibraries/MovieSceneSequenceExtensions.cpp
--- a/Engine/Plugins/MovieScene/SequencerScripting/Source/SequencerScripting/Private/ExtensionLibraries/MovieSceneSequenceExtensions.cpp
+++ b/Engine/Plugins/MovieScene/SequencerScripting/Source/SequencerScripting/Private/ExtensionLibraries/MovieSceneSequenceExtensions.cpp
@@ -267,6 +267,98 @@ void UMovieSceneSequenceExtensions::SetPlaybackEndSeconds(UMovieSceneSequence* S
 	}
 }
 
+void UMovieSceneSequenceExtensions::SetViewRangeStart(UMovieSceneSequence* Sequence, float StartTimeInSeconds)
+{
+	UMovieScene* MovieScene = GetMovieScene(Sequence);
+	if (MovieScene)
+	{
+#if WITH_EDITORONLY_DATA
+		MovieScene->SetViewRange(StartTimeInSeconds, MovieScene->GetEditorData().ViewEnd);
+#endif
+	}
+}
+
+float UMovieSceneSequenceExtensions::GetViewRangeStart(UMovieSceneSequence* Sequence)
+{
+	UMovieScene* MovieScene = GetMovieScene(Sequence);
+	if (MovieScene)
+	{
+#if WITH_EDITORONLY_DATA
+		return MovieScene->GetEditorData().ViewStart;
+#endif
+	}
+	return 0.f;
+}
+
+void UMovieSceneSequenceExtensions::SetViewRangeEnd(UMovieSceneSequence* Sequence, float EndTimeInSeconds)
+{
+	UMovieScene* MovieScene = GetMovieScene(Sequence);
+	if (MovieScene)
+	{
+#if WITH_EDITORONLY_DATA
+		MovieScene->SetViewRange(MovieScene->GetEditorData().ViewStart, EndTimeInSeconds);
+#endif
+	}
+}
+
+float UMovieSceneSequenceExtensions::GetViewRangeEnd(UMovieSceneSequence* Sequence)
+{
+	UMovieScene* MovieScene = GetMovieScene(Sequence);
+	if (MovieScene)
+	{
+#if WITH_EDITORONLY_DATA
+		return MovieScene->GetEditorData().ViewEnd;
+#endif
+	}
+	return 0.f;
+}
+
+void UMovieSceneSequenceExtensions::SetWorkRangeStart(UMovieSceneSequence* Sequence, float StartTimeInSeconds)
+{
+	UMovieScene* MovieScene = GetMovieScene(Sequence);
+	if (MovieScene)
+	{
+#if WITH_EDITORONLY_DATA
+		MovieScene->SetWorkingRange(StartTimeInSeconds, MovieScene->GetEditorData().WorkEnd);
+#endif
+	}
+}
+
+float UMovieSceneSequenceExtensions::GetWorkRangeStart(UMovieSceneSequence* Sequence)
+{
+	UMovieScene* MovieScene = GetMovieScene(Sequence);
+	if (MovieScene)
+	{
+#if WITH_EDITORONLY_DATA
+		return MovieScene->GetEditorData().WorkStart;
+#endif
+	}
+	return 0.f;
+}
+
+void UMovieSceneSequenceExtensions::SetWorkRangeEnd(UMovieSceneSequence* Sequence, float EndTimeInSeconds)
+{
+	UMovieScene* MovieScene = GetMovieScene(Sequence);
+	if (MovieScene)
+	{
+#if WITH_EDITORONLY_DATA
+		MovieScene->SetWorkingRange(MovieScene->GetEditorData().WorkStart, EndTimeInSeconds);
+#endif
+	}
+}
+
+float UMovieSceneSequenceExtensions::GetWorkRangeEnd(UMovieSceneSequence* Sequence)
+{
+	UMovieScene* MovieScene = GetMovieScene(Sequence);
+	if (MovieScene)
+	{
+#if WITH_EDITORONLY_DATA
+		return MovieScene->GetEditorData().WorkEnd;
+#endif
+	}
+	return 0.f;
+}
+
 FTimecode UMovieSceneSequenceExtensions::GetTimecodeSource(UMovieSceneSequence* Sequence)
 {
 	UMovieScene* MovieScene = GetMovieScene(Sequence);
diff --git a/Engine/Plugins/MovieScene/SequencerScripting/Source/SequencerScripting/Public/ExtensionLibraries/MovieSceneSequenceExtensions.h b/Engine/Plugins/MovieScene/SequencerScripting/Source/SequencerScripting/Public/ExtensionLibraries/MovieSceneSequenceExtensions.h
--- a/Engine/Plugins/MovieScene/SequencerScripting/Source/SequencerScripting/Public/ExtensionLibraries/MovieSceneSequenceExtensions.h
+++ b/Engine/Plugins/MovieScene/SequencerScripting/Source/SequencerScripting/Public/ExtensionLibraries/MovieSceneSequenceExtensions.h
@@ -202,6 +202,78 @@ public:
 	UFUNCTION(BlueprintCallable, Category = "Sequence", meta = (ScriptMethod))
 	static void SetPlaybackEndSeconds(UMovieSceneSequence* Sequence, float EndTime);
 
+	/**
+	 * Set the sequence view range start in seconds
+	 *
+	 * @param Sequence The sequence within which to set the view range start
+	 * @param StartTimeInSeconds The desired view range start time in seconds for this sequence
+	 */
+	UFUNCTION(BlueprintCallable, Category = "Sequence", meta = (ScriptMethod, DevelopmentOnly))
+	static void SetViewRangeStart(UMovieSceneSequence* InSequence, float StartTimeInSeconds);
+
+	/**
+	 * Get the sequence view range start in seconds
+	 *
+	 * @param Sequence The sequence within which to get the view range start
+	 * @return The view range start time in seconds for this sequence
+	 */
+	UFUNCTION(BlueprintPure, Category = "Sequence", meta = (ScriptMethod, DevelopmentOnly))
+	static float GetViewRangeStart(UMovieSceneSequence* InSequence);
+
+	/**
+	 * Set the sequence view range end in seconds
+	 *
+	 * @param Sequence The sequence within which to set the view range end
+	 * @param StartTimeInSeconds The desired view range end time in seconds for this sequence
+	 */
+	UFUNCTION(BlueprintCallable, Category = "Sequence", meta = (ScriptMethod, DevelopmentOnly))
+	static void SetViewRangeEnd(UMovieSceneSequence* InSequence, float EndTimeInSeconds);
+
+	/**
+	 * Get the sequence view range end in seconds
+	 *
+	 * @param Sequence The sequence within which to get the view range end
+	 * @return The view range end time in seconds for this sequence
+	 */
+	UFUNCTION(BlueprintPure, Category = "Sequence", meta = (ScriptMethod, DevelopmentOnly))
+	static float GetViewRangeEnd(UMovieSceneSequence* InSequence);
+
+	/**
+	 * Set the sequence work range start in seconds
+	 *
+	 * @param Sequence The sequence within which to set the work range start
+	 * @param StartTimeInSeconds The desired work range start time in seconds for this sequence
+	 */
+	UFUNCTION(BlueprintCallable, Category = "Sequence", meta = (ScriptMethod, DevelopmentOnly))
+	static void SetWorkRangeStart(UMovieSceneSequence* InSequence, float StartTimeInSeconds);
+
+	/**
+	 * Get the sequence work range start in seconds
+	 *
+	 * @param Sequence The sequence within which to get the work range start
+	 * @return The work range start time in seconds for this sequence
+	 */
+	UFUNCTION(BlueprintPure, Category = "Sequence", meta = (ScriptMethod, DevelopmentOnly))
+	static float GetWorkRangeStart(UMovieSceneSequence* InSequence);
+
+	/**
+	 * Set the sequence work range end in seconds
+	 *
+	 * @param Sequence The sequence within which to set the work range end
+	 * @param StartTimeInSeconds The desired work range end time in seconds for this sequence
+	 */
+	UFUNCTION(BlueprintCallable, Category = "Sequence", meta = (ScriptMethod, DevelopmentOnly))
+	static void SetWorkRangeEnd(UMovieSceneSequence* InSequence, float EndTimeInSeconds);
+
+	/**
+	 * Get the sequence work range end in seconds
+	 *
+	 * @param Sequence The sequence within which to get the work range end
+	 * @return The work range end time in seconds for this sequence
+	 */
+	UFUNCTION(BlueprintPure, Category = "Sequence", meta = (ScriptMethod, DevelopmentOnly))
+	static float GetWorkRangeEnd(UMovieSceneSequence* InSequence);
+
 	/**
 	 * Get the timecode source of this sequence
 	 *
