Sequencer: Fix duplicate Editor Scripting categories

diff --git a/Engine/Plugins/MovieScene/SequencerScripting/Source/SequencerScriptingEditor/Public/SequencerTools.h b/Engine/Plugins/MovieScene/SequencerScripting/Source/SequencerScriptingEditor/Public/SequencerTools.h
--- a/Engine/Plugins/MovieScene/SequencerScripting/Source/SequencerScriptingEditor/Public/SequencerTools.h
+++ b/Engine/Plugins/MovieScene/SequencerScripting/Source/SequencerScriptingEditor/Public/SequencerTools.h
@@ -28,13 +28,13 @@ public:
 	* if we're rendering via a PIE instance or a new process based on the passed in settings. Will return false
 	* if the state is not valid (ie: null or missing required parameters, capture in progress, etc.), true otherwise.
 	*/
-	UFUNCTION(BlueprintCallable, Category = "Editor Scripting | Sequencer Tools")
+	UFUNCTION(BlueprintCallable, Category = "Editor Scripting | Sequencer Tools | Movie Rendering")
 	static bool RenderMovie(class UMovieSceneCapture* InCaptureSettings, FOnRenderMovieStopped OnFinishedCallback);
 
 	/** 
 	* Returns if Render to Movie is currently in progress.
 	*/
-	UFUNCTION(BlueprintCallable, Category = "Editor Scripting | Sequencer Tools")
+	UFUNCTION(BlueprintCallable, Category = "Editor Scripting | Sequencer Tools | Movie Rendering")
 	static bool IsRenderingMovie()
 	{
 		IMovieSceneCaptureDialogModule& MovieSceneCaptureModule = FModuleManager::Get().LoadModuleChecked<IMovieSceneCaptureDialogModule>("MovieSceneCaptureDialog");
@@ -44,7 +44,7 @@ public:
 	/**
 	* Attempts to cancel an in-progress Render to Movie. Does nothing if there is no render in progress.
 	*/
-	UFUNCTION(BlueprintCallable, Category = "Editor Scripting | Sequencer Tools")
+	UFUNCTION(BlueprintCallable, Category = "Editor Scripting | Sequencer Tools | Movie Rendering")
 	static void CancelMovieRender();
 
 public:
@@ -56,7 +56,7 @@ public:
 	 * @InBindings Bindings to export
 	 * @InFBXFileName File to create
 	 */
-	UFUNCTION(BlueprintCallable, Category = "Editor Scripting | Sequencer Tools")
+	UFUNCTION(BlueprintCallable, Category = "Editor Scripting | Sequencer Tools | FBX")
 	static bool ExportFBX(UWorld* InWorld, ULevelSequence* InSequence, const TArray<FSequencerBindingProxy>& InBindings, UFbxExportOption* OverrideOptions,const FString& InFBXFileName);
 
 	/*
@@ -68,7 +68,7 @@ public:
 	 * @InImportFBXSettings Settings to control import.
 	 * @InImportFileName Path to fbx file to create
 	 */
-	UFUNCTION(BlueprintCallable, Category = "Editor Scripting | Sequencer Tools")
+	UFUNCTION(BlueprintCallable, Category = "Editor Scripting | Sequencer Tools | FBX")
 	static bool ImportFBX(UWorld* InWorld, ULevelSequence* InSequence, const TArray<FSequencerBindingProxy>& InBindings, UMovieSceneUserImportFBXSettings* InImportFBXSettings, const FString&  InImportFilename);
 
 
