Sequencer Scripting: Move extension libraries to Public folder so that the sequencer functions show up in blueprints

diff --git a/Engine/Plugins/MovieScene/SequencerScripting/Source/SequencerScripting/Private/ExtensionLibraries/MovieSceneBindingExtensions.h b/Engine/Plugins/MovieScene/SequencerScripting/Source/SequencerScripting/Private/ExtensionLibraries/MovieSceneBindingExtensions.h
deleted file mode 100644
diff --git a/Engine/Plugins/MovieScene/SequencerScripting/Source/SequencerScripting/Private/ExtensionLibraries/MovieScenePropertyTrackExtensions.h b/Engine/Plugins/MovieScene/SequencerScripting/Source/SequencerScripting/Private/ExtensionLibraries/MovieScenePropertyTrackExtensions.h
deleted file mode 100644
diff --git a/Engine/Plugins/MovieScene/SequencerScripting/Source/SequencerScripting/Private/ExtensionLibraries/MovieSceneSectionExtensions.h b/Engine/Plugins/MovieScene/SequencerScripting/Source/SequencerScripting/Private/ExtensionLibraries/MovieSceneSectionExtensions.h
deleted file mode 100644
diff --git a/Engine/Plugins/MovieScene/SequencerScripting/Source/SequencerScripting/Private/ExtensionLibraries/MovieSceneSequenceExtensions.h b/Engine/Plugins/MovieScene/SequencerScripting/Source/SequencerScripting/Private/ExtensionLibraries/MovieSceneSequenceExtensions.h
deleted file mode 100644
diff --git a/Engine/Plugins/MovieScene/SequencerScripting/Source/SequencerScripting/Private/ExtensionLibraries/MovieSceneSubSectionExtensions.h b/Engine/Plugins/MovieScene/SequencerScripting/Source/SequencerScripting/Private/ExtensionLibraries/MovieSceneSubSectionExtensions.h
deleted file mode 100644
diff --git a/Engine/Plugins/MovieScene/SequencerScripting/Source/SequencerScripting/Private/ExtensionLibraries/MovieSceneTrackExtensions.h b/Engine/Plugins/MovieScene/SequencerScripting/Source/SequencerScripting/Private/ExtensionLibraries/MovieSceneTrackExtensions.h
deleted file mode 100644
diff --git a/Engine/Plugins/MovieScene/SequencerScripting/Source/SequencerScripting/Private/ExtensionLibraries/SequencerScriptingRangeExtensions.h b/Engine/Plugins/MovieScene/SequencerScripting/Source/SequencerScripting/Private/ExtensionLibraries/SequencerScriptingRangeExtensions.h
deleted file mode 100644
diff --git a/Engine/Plugins/MovieScene/SequencerScripting/Source/SequencerScripting/Private/KeysAndChannels/MovieSceneScriptingChannel.h b/Engine/Plugins/MovieScene/SequencerScripting/Source/SequencerScripting/Private/KeysAndChannels/MovieSceneScriptingChannel.h
--- a/Engine/Plugins/MovieScene/SequencerScripting/Source/SequencerScripting/Private/KeysAndChannels/MovieSceneScriptingChannel.h
+++ b/Engine/Plugins/MovieScene/SequencerScripting/Source/SequencerScripting/Private/KeysAndChannels/MovieSceneScriptingChannel.h
@@ -12,7 +12,7 @@
 #include "KeyParams.h"
 #include "SequencerScriptingRange.h"
 #include "SequenceTimeUnit.h"
-#include "MovieSceneSequenceExtensions.h"
+#include "ExtensionLibraries/MovieSceneSequenceExtensions.h"
 
 #include "MovieSceneScriptingChannel.generated.h"
 
diff --git a/Engine/Plugins/MovieScene/SequencerScripting/Source/SequencerScripting/Public/ExtensionLibraries/MovieSceneBindingExtensions.h b/Engine/Plugins/MovieScene/SequencerScripting/Source/SequencerScripting/Public/ExtensionLibraries/MovieSceneBindingExtensions.h
new file mode 100644
--- /dev/null
+++ b/Engine/Plugins/MovieScene/SequencerScripting/Source/SequencerScripting/Public/ExtensionLibraries/MovieSceneBindingExtensions.h
@@ -0,0 +1,147 @@
+// Copyright 1998-2018 Epic Games, Inc. All Rights Reserved.
+
+#pragma once
+
+#include "CoreTypes.h"
+#include "UObject/ObjectMacros.h"
+#include "Kismet/BlueprintFunctionLibrary.h"
+#include "MovieSceneTrack.h"
+#include "Templates/SubclassOf.h"
+#include "SequencerBindingProxy.h"
+
+#include "MovieSceneBindingExtensions.generated.h"
+
+
+/**
+ * Function library containing methods that should be hoisted onto FMovieSceneBindingProxies for scripting
+ */
+UCLASS()
+class UMovieSceneBindingExtensions : public UBlueprintFunctionLibrary
+{
+	GENERATED_BODY()
+public:
+	/**
+	 * Check whether the specified binding is valid
+	 */
+	UFUNCTION(BlueprintPure, Category=Sequence, meta=(ScriptMethod))
+	static bool IsValid(const FSequencerBindingProxy& InBinding);
+
+	/**
+	 * Get this binding's ID
+	 *
+	 * @param InBinding     The binding to get the ID of
+	 * @return The guid that uniquely represents this binding
+	 */
+	UFUNCTION(BlueprintCallable, Category=Sequence, meta=(ScriptMethod))
+	static FGuid GetId(const FSequencerBindingProxy& InBinding);
+
+	/**
+	 * Get this binding's name
+	 *
+	 * @param InBinding     The binding to get the name of
+	 * @return The display name of the binding
+	 */
+	UFUNCTION(BlueprintCallable, Category=Sequence, meta=(ScriptMethod))
+	static FText GetDisplayName(const FSequencerBindingProxy& InBinding);
+
+	/**
+	 * Get this binding's object non-display name
+	 *
+	 * @param InBinding     The binding to get the name of
+	 * @return The name of the binding
+	 */
+	UFUNCTION(BlueprintCallable, Category=Sequence, meta=(ScriptMethod))
+	static FString GetName(const FSequencerBindingProxy& InBinding);
+
+	/**
+	 * Get all the tracks stored within this binding
+	 *
+	 * @param InBinding     The binding to find tracks in
+	 * @return An array containing all the binding's tracks
+	 */
+	UFUNCTION(BlueprintCallable, Category=Sequence, meta=(ScriptMethod))
+	static TArray<UMovieSceneTrack*> GetTracks(const FSequencerBindingProxy& InBinding);
+
+	/**
+	 * Find all tracks within a given binding of the specified type
+	 *
+	 * @param InBinding     The binding to find tracks in
+	 * @param TrackType     A UMovieSceneTrack class type specifying which types of track to return
+	 * @return An array containing any tracks that match the type specified
+	 */
+	UFUNCTION(BlueprintCallable, Category=Sequence, meta=(ScriptMethod))
+	static TArray<UMovieSceneTrack*> FindTracksByType(const FSequencerBindingProxy& InBinding, TSubclassOf<UMovieSceneTrack> TrackType);
+
+	/**
+	 * Find all tracks within a given binding of the specified type, not allowing sub-classed types
+	 *
+	 * @param InBinding     The binding to find tracks in
+	 * @param TrackType     A UMovieSceneTrack class type specifying the exact types of track to return
+	 * @return An array containing any tracks that are exactly the same as the type specified
+	 */
+	UFUNCTION(BlueprintCallable, Category=Sequence, meta=(ScriptMethod))
+	static TArray<UMovieSceneTrack*> FindTracksByExactType(const FSequencerBindingProxy& InBinding, TSubclassOf<UMovieSceneTrack> TrackType);
+
+	/**
+	 * Remove the specified track from this binding
+	 *
+	 * @param InBinding     The binding to remove the track from
+	 * @param TrackToRemove The track to remove
+	 */
+	UFUNCTION(BlueprintCallable, Category=Sequence, meta=(ScriptMethod))
+	static void RemoveTrack(const FSequencerBindingProxy& InBinding, UMovieSceneTrack* TrackToRemove);
+
+	/**
+	 * Remove the specified binding
+	 *
+	 * @param InBinding     The binding to remove the track from
+	 */
+	UFUNCTION(BlueprintCallable, Category=Sequence, meta=(ScriptMethod))
+	static void Remove(const FSequencerBindingProxy& InBinding);
+
+	/**
+	 * Add a new track to the specified binding
+	 *
+	 * @param InBinding     The binding to add tracks to
+	 * @param TrackType     A UMovieSceneTrack class type specifying the type of track to create
+	 * @return The newly created track, if successful
+	 */
+	UFUNCTION(BlueprintCallable, Category=Sequence, meta=(ScriptMethod))
+	static UMovieSceneTrack* AddTrack(const FSequencerBindingProxy& InBinding, TSubclassOf<UMovieSceneTrack> TrackType);
+
+	/**
+	* Get all the children of this binding
+	*
+	* @param InBinding     The binding to to get children of
+	* @return An array containing all the binding's children
+	*/
+	UFUNCTION(BlueprintCallable, Category=Sequence, meta=(ScriptMethod))
+	static TArray<FSequencerBindingProxy> GetChildPossessables(const FSequencerBindingProxy& InBinding);
+
+	/**
+	* Get this binding's object template
+	*
+	* @param InBinding     The binding to get the object template of
+	* @return The object template of the binding
+	*/
+	UFUNCTION(BlueprintCallable, Category=Sequence, meta=(ScriptMethod))
+	static UObject* GetObjectTemplate(const FSequencerBindingProxy& InBinding);
+
+	/**
+	* Get this binding's possessed object class
+	*
+	* @param InBinding     The binding to get the possessed object class of
+	* @return The possessed object class of the binding
+	*/
+	UFUNCTION(BlueprintCallable, Category=Sequence, meta=(ScriptMethod))
+	static UClass* GetPossessedObjectClass(const FSequencerBindingProxy& InBinding);
+
+	/**
+	* Get the parent of this binding
+	*
+	* @param InBinding     The binding to get the parent of
+	* @return The binding's parent
+	*/
+	UFUNCTION(BlueprintCallable, Category=Sequence, meta=(ScriptMethod))
+	static FSequencerBindingProxy GetParent(const FSequencerBindingProxy& InBinding);
+};
\ No newline at end of file
diff --git a/Engine/Plugins/MovieScene/SequencerScripting/Source/SequencerScripting/Public/ExtensionLibraries/MovieScenePropertyTrackExtensions.h b/Engine/Plugins/MovieScene/SequencerScripting/Source/SequencerScripting/Public/ExtensionLibraries/MovieScenePropertyTrackExtensions.h
new file mode 100644
--- /dev/null
+++ b/Engine/Plugins/MovieScene/SequencerScripting/Source/SequencerScripting/Public/ExtensionLibraries/MovieScenePropertyTrackExtensions.h
@@ -0,0 +1,50 @@
+// Copyright 1998-2018 Epic Games, Inc. All Rights Reserved.
+
+#pragma once
+
+#include "CoreTypes.h"
+#include "UObject/ObjectMacros.h"
+#include "Kismet/BlueprintFunctionLibrary.h"
+
+#include "MovieScenePropertyTrackExtensions.generated.h"
+
+
+class UMovieScenePropertyTrack;
+
+/**
+ * Function library containing methods that should be hoisted onto UMovieScenePropertyTrack for scripting
+ */
+UCLASS()
+class UMovieScenePropertyTrackExtensions : public UBlueprintFunctionLibrary
+{
+	GENERATED_BODY()
+
+public:
+
+	/**
+	 * Get this track's property name
+	 *
+	 * @param Track        The track to use
+	 * @return This track's property name
+	 */
+	UFUNCTION(BlueprintCallable, Category="Track", meta=(ScriptMethod))
+	static FName GetPropertyName(UMovieScenePropertyTrack* Track);
+
+	/**
+	 * Get this track's property path
+	 *
+	 * @param Track        The track to use
+	 * @return This track's property path
+	 */
+	UFUNCTION(BlueprintCallable, Category="Track", meta=(ScriptMethod))
+	static FString GetPropertyPath(UMovieScenePropertyTrack* Track);
+
+	/**
+	 * Get this track's unique name
+	 *
+	 * @param Track        The track to use
+	 * @return This track's unique name
+	 */
+	UFUNCTION(BlueprintCallable, Category="Track", meta=(ScriptMethod))
+	static FName GetUniqueTrackName(UMovieScenePropertyTrack* Track);
+};
diff --git a/Engine/Plugins/MovieScene/SequencerScripting/Source/SequencerScripting/Public/ExtensionLibraries/MovieSceneSectionExtensions.h b/Engine/Plugins/MovieScene/SequencerScripting/Source/SequencerScripting/Public/ExtensionLibraries/MovieSceneSectionExtensions.h
new file mode 100644
--- /dev/null
+++ b/Engine/Plugins/MovieScene/SequencerScripting/Source/SequencerScripting/Public/ExtensionLibraries/MovieSceneSectionExtensions.h
@@ -0,0 +1,59 @@
+// Copyright 1998-2018 Epic Games, Inc. All Rights Reserved.
+
+#pragma once
+
+#include "CoreTypes.h"
+#include "Templates/SubclassOf.h"
+#include "UObject/ObjectMacros.h"
+#include "Kismet/BlueprintFunctionLibrary.h"
+#include "SequencerScriptingRange.h"
+#include "MovieSceneScriptingChannel.h"
+#include "MovieSceneSectionExtensions.generated.h"
+
+class UMovieSceneSection;
+
+/**
+ * Function library containing methods that should be hoisted onto UMovieSceneSections for scripting
+ */
+UCLASS()
+class UMovieSceneSectionExtensions : public UBlueprintFunctionLibrary
+{
+	GENERATED_BODY()
+public:
+	/**
+	 * Get the range of this section
+	 */
+	UFUNCTION(BlueprintCallable, Category=Section, meta=(ScriptMethod))
+	static FSequencerScriptingRange GetRange(UMovieSceneSection* Section);
+
+	/**
+	 * Set the range of this section
+	 *
+	 * @param Range         The desired range for this section
+	 */
+	UFUNCTION(BlueprintCallable, Category=Section, meta=(ScriptMethod))
+	static void SetRange(UMovieSceneSection* Section, const FSequencerScriptingRange& Range);
+
+	/**
+	* Find all channels that belong to the specified UMovieSceneSection. Some sections have many channels (such as
+	* Transforms containing 9 float channels to represent Translation/Rotation/Scale), and a section may have mixed
+	* channel types.
+	*
+	* @param Section       The section to use.
+	* @return An array containing any key channels that match the type specified
+	*/
+	UFUNCTION(BlueprintCallable, Category = Section, meta=(ScriptMethod))
+	static TArray<UMovieSceneScriptingChannel*> GetChannels(UMovieSceneSection* Section);
+
+	/**
+	* Find all channels that belong to the specified UMovieSceneSection that match the specific type. This will filter out any children who do not inherit
+	* from the specified type for you.
+	*
+	* @param Section        The section to use.
+	* @param ChannelType	The class type to look for.
+	* @return An array containing any key channels that match the type specified
+	*/
+	UFUNCTION(BlueprintCallable, Category = Section, meta = (ScriptMethod))
+	static TArray<UMovieSceneScriptingChannel*> FindChannelsByType(UMovieSceneSection* Section, TSubclassOf<UMovieSceneScriptingChannel> ChannelType);
+
+};
\ No newline at end of file
diff --git a/Engine/Plugins/MovieScene/SequencerScripting/Source/SequencerScripting/Public/ExtensionLibraries/MovieSceneSequenceExtensions.h b/Engine/Plugins/MovieScene/SequencerScripting/Source/SequencerScripting/Public/ExtensionLibraries/MovieSceneSequenceExtensions.h
new file mode 100644
--- /dev/null
+++ b/Engine/Plugins/MovieScene/SequencerScripting/Source/SequencerScripting/Public/ExtensionLibraries/MovieSceneSequenceExtensions.h
@@ -0,0 +1,210 @@
+// Copyright 1998-2018 Epic Games, Inc. All Rights Reserved.
+
+#pragma once
+
+#include "CoreTypes.h"
+#include "UObject/ObjectMacros.h"
+#include "Containers/ArrayView.h"
+#include "Kismet/BlueprintFunctionLibrary.h"
+#include "SequencerBindingProxy.h"
+#include "SequencerScriptingRange.h"
+#include "Templates/SubclassOf.h"
+#include "MovieSceneTrack.h"
+
+#include "MovieSceneSequenceExtensions.generated.h"
+
+/**
+ * Function library containing methods that should be hoisted onto UMovieSceneSequences for scripting purposes
+ */
+UCLASS()
+class UMovieSceneSequenceExtensions : public UBlueprintFunctionLibrary
+{
+	GENERATED_BODY()
+public:
+	/**
+	 * Get this sequence's movie scene data
+	 *
+	 * @param Sequence        The sequence to use
+	 * @return This sequence's movie scene data object
+	 */
+	UFUNCTION(BlueprintCallable, Category="Sequence", meta=(ScriptMethod))
+	static UMovieScene* GetMovieScene(UMovieSceneSequence* Sequence);
+
+	/**
+	 * Get all master tracks
+	 *
+	 * @param Sequence        The sequence to use
+	 * @return An array containing all master tracks in this sequence
+	 */
+	UFUNCTION(BlueprintCallable, Category="Sequence", meta=(ScriptMethod))
+	static TArray<UMovieSceneTrack*> GetMasterTracks(UMovieSceneSequence* Sequence);
+
+	/**
+	 * Find all master tracks of the specified type
+	 *
+	 * @param Sequence        The sequence to use
+	 * @param TrackType     A UMovieSceneTrack class type specifying which types of track to return
+	 * @return An array containing any tracks that match the type specified
+	 */
+	UFUNCTION(BlueprintCallable, Category="Sequence", meta=(ScriptMethod))
+	static TArray<UMovieSceneTrack*> FindMasterTracksByType(UMovieSceneSequence* Sequence, TSubclassOf<UMovieSceneTrack> TrackType);
+
+	/**
+	 * Find all master tracks of the specified type, not allowing sub-classed types
+	 *
+	 * @param Sequence        The sequence to use
+	 * @param TrackType     A UMovieSceneTrack class type specifying the exact types of track to return
+	 * @return An array containing any tracks that are exactly the same as the type specified
+	 */
+	UFUNCTION(BlueprintCallable, Category="Sequence", meta=(ScriptMethod))
+	static TArray<UMovieSceneTrack*> FindMasterTracksByExactType(UMovieSceneSequence* Sequence, TSubclassOf<UMovieSceneTrack> TrackType);
+
+	/**
+	 * Add a new master track of the specified type
+	 *
+	 * @param Sequence        The sequence to use
+	 * @param TrackType     A UMovieSceneTrack class type to create
+	 * @return The newly created track, if successful
+	 */
+	UFUNCTION(BlueprintCallable, Category="Sequence", meta=(ScriptMethod))
+	static UMovieSceneTrack* AddMasterTrack(UMovieSceneSequence* Sequence, TSubclassOf<UMovieSceneTrack> TrackType);
+
+	/**
+	 * Get's this sequence's display rate
+	 *
+	 * @param Sequence        The sequence to use
+	 * @return The display rate that this sequence is displayed as
+	 */
+	UFUNCTION(BlueprintCallable, Category="Sequence", meta=(ScriptMethod))
+	static FFrameRate GetDisplayRate(UMovieSceneSequence* Sequence);
+
+	/**
+	 * Get's this sequence's tick resolution
+	 *
+	 * @param Sequence        The sequence to use
+	 * @return The tick resolution of the sequence, defining the smallest unit of time representable on this sequence
+	 */
+	UFUNCTION(BlueprintCallable, Category="Sequence", meta=(ScriptMethod))
+	static FFrameRate GetTickResolution(UMovieSceneSequence* Sequence);
+
+	/**
+	 * Make a new range for this sequence in its display rate
+	 *
+	 * @param Sequence        The sequence within which to find the binding
+	 * @param StartFrame      The frame at which to start the range
+	 * @param Duration        The length of the range
+	 * @return Specified sequencer range
+	 */
+	UFUNCTION(BlueprintCallable, Category="Sequence", meta=(ScriptMethod))
+	static FSequencerScriptingRange MakeRange(UMovieSceneSequence* Sequence, int32 StartFrame, int32 Duration);
+
+	/**
+	 * Make a new range for this sequence in seconds
+	 *
+	 * @param Sequence        The sequence within which to find the binding
+	 * @param StartTime       The time in seconds at which to start the range
+	 * @param Duration        The length of the range in seconds
+	 * @return Specified sequencer range
+	 */
+	UFUNCTION(BlueprintCallable, Category="Sequence", meta=(ScriptMethod))
+	static FSequencerScriptingRange MakeRangeSeconds(UMovieSceneSequence* Sequence, float StartTime, float Duration);
+
+	/**
+	 * Get playback range of this sequence
+	 *
+	 * @param Sequence        The sequence within which to get the playback range
+	 * @return Playback range of this sequence
+	 */
+	UFUNCTION(BlueprintCallable, Category="Sequence", meta=(ScriptMethod))
+	static FSequencerScriptingRange GetPlaybackRange(UMovieSceneSequence* Sequence);
+
+	/**
+	 * Attempt to locate a binding in this sequence by its name
+	 *
+	 * @param Sequence        The sequence within which to find the binding
+	 * @param Name            The display name of the binding to look up
+	 * @return A unique identifier for the binding, or invalid
+	 */
+	UFUNCTION(BlueprintCallable, Category="Sequence", meta=(ScriptMethod))
+	static FSequencerBindingProxy FindBindingByName(UMovieSceneSequence* Sequence, FString Name);
+
+	/**
+	 * Get all the bindings in this sequence
+	 *
+	 * @param Sequence        The sequence to get bindings for
+	 * @return An array of unique identifiers for all the bindings in this sequence
+	 */
+	UFUNCTION(BlueprintCallable, Category="Sequence", meta=(ScriptMethod))
+	static TArray<FSequencerBindingProxy> GetBindings(UMovieSceneSequence* Sequence);
+
+	/**
+	* Get all the spawnables in this sequence
+	*
+	* @param Sequence        The sequence to get spawnables for
+	* @return Spawnables in this sequence
+	*/
+	UFUNCTION(BlueprintCallable, Category="Sequence", meta=(ScriptMethod))
+	static TArray<FSequencerBindingProxy> GetSpawnables(UMovieSceneSequence* Sequence);
+
+	/**
+	* Get all the possessables in this sequence
+	*
+	* @param Sequence        The sequence to get possessables for
+	* @return Possessables in this sequence
+	*/
+	UFUNCTION(BlueprintCallable, Category="Sequence", meta=(ScriptMethod))
+	static TArray<FSequencerBindingProxy> GetPossessables(UMovieSceneSequence* Sequence);
+
+	/**
+	 * Add a new binding to this sequence that will possess the specified object
+	 *
+	 * @param Sequence        The sequence to add a possessable to
+	 * @param ObjectToPossess The object that this sequence should possess when evaluating
+	 * @return A unique identifier for the new binding
+	 */
+	UFUNCTION(BlueprintCallable, Category="Sequence", meta=(ScriptMethod))
+	static FSequencerBindingProxy AddPossessable(UMovieSceneSequence* Sequence, UObject* ObjectToPossess);
+
+	/**
+	 * Add a new binding to this sequence that will spawn the specified object
+	 *
+	 * @param Sequence        The sequence to add to
+	 * @param ObjectToSpawn   An object instance to use as a template for spawning
+	 * @return A unique identifier for the new binding
+	 */
+	UFUNCTION(BlueprintCallable, Category="Sequence", meta=(ScriptMethod))
+	static FSequencerBindingProxy AddSpawnableFromInstance(UMovieSceneSequence* Sequence, UObject* ObjectToSpawn);
+
+	/**
+	 * Add a new binding to this sequence that will spawn the specified object
+	 *
+	 * @param Sequence        The sequence to add to
+	 * @param ClassToSpawn    A class or blueprint type to spawn for this binding
+	 * @return A unique identifier for the new binding
+	 */
+	UFUNCTION(BlueprintCallable, Category="Sequence", meta=(ScriptMethod))
+	static FSequencerBindingProxy AddSpawnableFromClass(UMovieSceneSequence* Sequence, UClass* ClassToSpawn);
+
+	/**
+	 * Locate all the objects that correspond to the specified object ID, using the specified context
+	 *
+	 * @param Sequence   The sequence to locate bound objects for
+	 * @param InBinding  The object binding
+	 * @param Context    Optional context to use to find the required object
+	 * @return An array of all bound objects
+	 */
+	UFUNCTION(BlueprintCallable, Category="Sequence", meta=(ScriptMethod))
+	static TArray<UObject*> LocateBoundObjects(UMovieSceneSequence* Sequence, const FSequencerBindingProxy& InBinding, UObject* Context);
+
+public:
+
+	/**
+	 * Filter the specified array of tracks by class, optionally only matching exact classes
+	 *
+	 * @param InTracks       The array of tracks to filter
+	 * @param DesiredClass   The class to match against
+	 * @param bExactMatch    Whether to match sub classes or not
+	 * @return A filtered array of tracks
+	 */
+	static TArray<UMovieSceneTrack*> FilterTracks(TArrayView<UMovieSceneTrack* const> InTracks, UClass* DesiredClass, bool bExactMatch);
+};
\ No newline at end of file
diff --git a/Engine/Plugins/MovieScene/SequencerScripting/Source/SequencerScripting/Public/ExtensionLibraries/MovieSceneSubSectionExtensions.h b/Engine/Plugins/MovieScene/SequencerScripting/Source/SequencerScripting/Public/ExtensionLibraries/MovieSceneSubSectionExtensions.h
new file mode 100644
--- /dev/null
+++ b/Engine/Plugins/MovieScene/SequencerScripting/Source/SequencerScripting/Public/ExtensionLibraries/MovieSceneSubSectionExtensions.h
@@ -0,0 +1,28 @@
+// Copyright 1998-2018 Epic Games, Inc. All Rights Reserved.
+
+#pragma once
+
+#include "CoreTypes.h"
+#include "Templates/SubclassOf.h"
+#include "UObject/ObjectMacros.h"
+#include "Kismet/BlueprintFunctionLibrary.h"
+#include "MovieSceneSubSectionExtensions.generated.h"
+
+class UMovieSceneSubSection;
+class UMovieSceneSequence;
+
+/**
+ * Function library containing methods that should be hoisted onto UMovieSceneSections for scripting
+ */
+UCLASS()
+class UMovieSceneSubSectionExtensions : public UBlueprintFunctionLibrary
+{
+	GENERATED_BODY()
+public:
+
+	/**
+	 * Get sequence of this section
+	 */
+	UFUNCTION(BlueprintCallable, Category=Section, meta=(ScriptMethod))
+	static UMovieSceneSequence* GetSequence(UMovieSceneSubSection* SubSection);
+};
diff --git a/Engine/Plugins/MovieScene/SequencerScripting/Source/SequencerScripting/Public/ExtensionLibraries/MovieSceneTrackExtensions.h b/Engine/Plugins/MovieScene/SequencerScripting/Source/SequencerScripting/Public/ExtensionLibraries/MovieSceneTrackExtensions.h
new file mode 100644
--- /dev/null
+++ b/Engine/Plugins/MovieScene/SequencerScripting/Source/SequencerScripting/Public/ExtensionLibraries/MovieSceneTrackExtensions.h
@@ -0,0 +1,61 @@
+// Copyright 1998-2018 Epic Games, Inc. All Rights Reserved.
+
+#pragma once
+
+#include "CoreTypes.h"
+#include "UObject/ObjectMacros.h"
+#include "Kismet/BlueprintFunctionLibrary.h"
+
+#include "MovieSceneTrackExtensions.generated.h"
+
+class FText;
+
+class UMovieSceneTrack;
+class UMovieSceneSection;
+
+/**
+ * Function library containing methods that should be hoisted onto UMovieSceneTracks for scripting
+ */
+UCLASS()
+class UMovieSceneTrackExtensions : public UBlueprintFunctionLibrary
+{
+	GENERATED_BODY()
+
+public:
+
+	/**
+	 * Get this track's display name
+	 *
+	 * @param Track        The track to use
+	 * @return This track's display name
+	 */
+	UFUNCTION(BlueprintCallable, Category="Track", meta=(ScriptMethod))
+	static FText GetDisplayName(UMovieSceneTrack* Track);
+
+	/**
+	 * Add a new section to this track
+	 *
+	 * @param Track        The track to use
+	 * @return The newly create section if successful
+	 */
+	UFUNCTION(BlueprintCallable, Category="Track", meta=(ScriptMethod))
+	static UMovieSceneSection* AddSection(UMovieSceneTrack* Track);
+
+	/**
+	 * Access all this track's sections
+	 *
+	 * @param Track        The track to use
+	 * @return An array of this track's sections
+	 */
+	UFUNCTION(BlueprintCallable, Category="Track", meta=(ScriptMethod))
+	static TArray<UMovieSceneSection*> GetSections(UMovieSceneTrack* Track);
+
+	/**
+	 * Remove the specified section
+	 *
+	 * @param Track        The track to remove the section from, if present
+	 * @param Section      The section to remove
+	 */
+	UFUNCTION(BlueprintCallable, Category="Track", meta=(ScriptMethod))
+	static void RemoveSection(UMovieSceneTrack* Track, UMovieSceneSection* Section);
+};
\ No newline at end of file
diff --git a/Engine/Plugins/MovieScene/SequencerScripting/Source/SequencerScripting/Public/ExtensionLibraries/SequencerScriptingRangeExtensions.h b/Engine/Plugins/MovieScene/SequencerScripting/Source/SequencerScripting/Public/ExtensionLibraries/SequencerScriptingRangeExtensions.h
new file mode 100644
--- /dev/null
+++ b/Engine/Plugins/MovieScene/SequencerScripting/Source/SequencerScripting/Public/ExtensionLibraries/SequencerScriptingRangeExtensions.h
@@ -0,0 +1,115 @@
+// Copyright 1998-2018 Epic Games, Inc. All Rights Reserved.
+
+#pragma once
+
+#include "CoreTypes.h"
+#include "UObject/ObjectMacros.h"
+#include "Kismet/BlueprintFunctionLibrary.h"
+#include "SequencerScriptingRange.h"
+
+#include "SequencerScriptingRangeExtensions.generated.h"
+
+/**
+ * Function library containing methods that should be hoisted onto FSequencerScriptingRanges
+ */
+UCLASS()
+class USequencerScriptingRangeExtensions : public UBlueprintFunctionLibrary
+{
+	GENERATED_BODY()
+public:
+	/**
+	 * Check whether this range has a start
+	 *
+	 * @param Range       The range to check
+	 */
+	UFUNCTION(BlueprintCallable, Category=Sequence, meta=(ScriptMethod))
+	static bool HasStart(const FSequencerScriptingRange& Range);
+
+	/**
+	 * Check whether this range has an end
+	 *
+	 * @param Range       The range to check
+	 */
+	UFUNCTION(BlueprintCallable, Category=Sequence, meta=(ScriptMethod))
+	static bool HasEnd(const FSequencerScriptingRange& Range);
+
+	/**
+	 * Remove the start from this range, making it infinite
+	 *
+	 * @param Range       The range to remove the start from
+	 */
+	UFUNCTION(BlueprintCallable, Category=Sequence, meta=(ScriptMethod))
+	static void RemoveStart(UPARAM(ref) FSequencerScriptingRange& Range);
+
+	/**
+	 * Remove the end from this range, making it infinite
+	 *
+	 * @param Range       The range to remove the end from
+	 */
+	UFUNCTION(BlueprintCallable, Category=Sequence, meta=(ScriptMethod))
+	static void RemoveEnd(UPARAM(ref) FSequencerScriptingRange& Range);
+
+	/**
+	 * Get the starting time for the specified range in seconds, if it has one. Defined as the first valid time that is inside the range.
+	 *
+	 * @param Range       The range to get the start from
+	 */
+	UFUNCTION(BlueprintCallable, Category=Sequence, meta=(ScriptMethod))
+	static float GetStartSeconds(const FSequencerScriptingRange& Range);
+
+	/**
+	 * Get the ending time for the specified range in seconds, if it has one. Defined as the first time that is outside of the range.
+	 *
+	 * @param Range       The range to get the end from
+	 */
+	UFUNCTION(BlueprintCallable, Category=Sequence, meta=(ScriptMethod))
+	static float GetEndSeconds(const FSequencerScriptingRange& Range);
+
+	/**
+	 * Set the starting time for the specified range in seconds. Interpreted as the first valid time that is inside the range.
+	 *
+	 * @param Range       The range to set the start on
+	 */
+	UFUNCTION(BlueprintCallable, Category=Sequence, meta=(ScriptMethod))
+	static void SetStartSeconds(UPARAM(ref) FSequencerScriptingRange& Range, float Start);
+
+	/**
+	 * Set the ending time for the specified range in seconds. Interpreted as the first time that is outside of the range.
+	 *
+	 * @param Range       The range to set the end on
+	 */
+	UFUNCTION(BlueprintCallable, Category=Sequence, meta=(ScriptMethod))
+	static void SetEndSeconds(UPARAM(ref) FSequencerScriptingRange& Range, float End);
+
+	/**
+	 * Get the starting frame for the specified range, if it has one. Defined as the first valid frame that is inside the range.
+	 *
+	 * @param Range       The range to get the start from
+	 */
+	UFUNCTION(BlueprintCallable, Category=Sequence, meta=(ScriptMethod))
+	static int32 GetStartFrame(const FSequencerScriptingRange& Range);
+
+	/**
+	 * Get the ending frame for the specified range, if it has one. Defined as the first subsequent frame that is outside of the range.
+	 *
+	 * @param Range       The range to get the end from
+	 */
+	UFUNCTION(BlueprintCallable, Category=Sequence, meta=(ScriptMethod))
+	static int32 GetEndFrame(const FSequencerScriptingRange& Range);
+
+	/**
+	 * Set the starting frame for the specified range. Interpreted as the first valid frame that is inside the range.
+	 *
+	 * @param Range       The range to set the start on
+	 */
+	UFUNCTION(BlueprintCallable, Category=Sequence, meta=(ScriptMethod))
+	static void SetStartFrame(UPARAM(ref) FSequencerScriptingRange& Range, int32 Start);
+
+	/**
+	 * Set the ending frame for the specified range. Interpreted as the first subsequent frame that is outside of the range.
+	 *
+	 * @param Range       The range to set the end on
+	 */
+	UFUNCTION(BlueprintCallable, Category=Sequence, meta=(ScriptMethod))
+	static void SetEndFrame(UPARAM(ref) FSequencerScriptingRange& Range, int32 End);
+};
\ No newline at end of file
