Sequencer Scripting: Set as developer type plugin so that blueprint library functions show up.

diff --git a/Engine/Plugins/MovieScene/SequencerScripting/SequencerScripting.uplugin b/Engine/Plugins/MovieScene/SequencerScripting/SequencerScripting.uplugin
--- a/Engine/Plugins/MovieScene/SequencerScripting/SequencerScripting.uplugin
+++ b/Engine/Plugins/MovieScene/SequencerScripting/SequencerScripting.uplugin
@@ -22,8 +22,8 @@
 	"Modules": [
 		{
 			"Name": "SequencerScripting",
-			"Type": "Editor",
-			"LoadingPhase": "Default"
+			"Type": "Developer",
+			"LoadingPhase": "PostEngineInit"
 		}
 	]
 }
\ No newline at end of file
