Sequencer Scripting: Added GetTimecodeSource function

diff --git a/Engine/Plugins/MovieScene/SequencerScripting/Source/SequencerScripting/Private/ExtensionLibraries/MovieSceneSequenceExtensions.cpp b/Engine/Plugins/MovieScene/SequencerScripting/Source/SequencerScripting/Private/ExtensionLibraries/MovieSceneSequenceExtensions.cpp
--- a/Engine/Plugins/MovieScene/SequencerScripting/Source/SequencerScripting/Private/ExtensionLibraries/MovieSceneSequenceExtensions.cpp
+++ b/Engine/Plugins/MovieScene/SequencerScripting/Source/SequencerScripting/Private/ExtensionLibraries/MovieSceneSequenceExtensions.cpp
@@ -144,6 +144,19 @@ FSequencerScriptingRange UMovieSceneSequenceExtensions::GetPlaybackRange(UMovieS
 	}
 }
 
+FTimecode UMovieSceneSequenceExtensions::GetTimecodeSource(UMovieSceneSequence* Sequence)
+{
+	UMovieScene* MovieScene = Sequence->GetMovieScene();
+	if (MovieScene)
+	{
+		return MovieScene->TimecodeSource.Timecode;
+	}
+	else
+	{
+		return FTimecode();
+	}
+}
+
 FSequencerBindingProxy UMovieSceneSequenceExtensions::FindBindingByName(UMovieSceneSequence* Sequence, FString Name)
 {
 	UMovieScene* MovieScene = Sequence->GetMovieScene();
diff --git a/Engine/Plugins/MovieScene/SequencerScripting/Source/SequencerScripting/Public/ExtensionLibraries/MovieSceneSequenceExtensions.h b/Engine/Plugins/MovieScene/SequencerScripting/Source/SequencerScripting/Public/ExtensionLibraries/MovieSceneSequenceExtensions.h
--- a/Engine/Plugins/MovieScene/SequencerScripting/Source/SequencerScripting/Public/ExtensionLibraries/MovieSceneSequenceExtensions.h
+++ b/Engine/Plugins/MovieScene/SequencerScripting/Source/SequencerScripting/Public/ExtensionLibraries/MovieSceneSequenceExtensions.h
@@ -118,6 +118,15 @@ public:
 	UFUNCTION(BlueprintCallable, Category="Sequence", meta=(ScriptMethod))
 	static FSequencerScriptingRange GetPlaybackRange(UMovieSceneSequence* Sequence);
 
+	/**
+	* Get the timecode source of this sequence
+	*
+	* @param Sequence        The sequence within which to get the timecode source
+	* @return Timecode source of this sequence
+	*/
+	UFUNCTION(BlueprintCallable, Category = "Sequence", meta = (ScriptMethod))
+	static FTimecode GetTimecodeSource(UMovieSceneSequence* Sequence);
+
 	/**
 	 * Attempt to locate a binding in this sequence by its name
 	 *
