Back out changelist 4500571 because it depends on PythonScriptPlugin which is an editor plugin

diff --git a/Engine/Plugins/MovieScene/SequencerScripting/SequencerScripting.uplugin b/Engine/Plugins/MovieScene/SequencerScripting/SequencerScripting.uplugin
--- a/Engine/Plugins/MovieScene/SequencerScripting/SequencerScripting.uplugin
+++ b/Engine/Plugins/MovieScene/SequencerScripting/SequencerScripting.uplugin
@@ -22,8 +22,8 @@
 	"Modules": [
 		{
 			"Name": "SequencerScripting",
-			"Type": "Developer",
-			"LoadingPhase": "PostEngineInit"
+			"Type": "Editor",
+			"LoadingPhase": "Default"
 		}
 	]
 }
\ No newline at end of file
