Sequencer Scripting: Make SequencerScripting a Developer module so that its functions are exposed in the level blueprint. Move editor scripting functions (ie. render movie) to Sequencer Scripting Editor.

diff --git a/Engine/Plugins/MovieScene/SequencerScripting/SequencerScripting.uplugin b/Engine/Plugins/MovieScene/SequencerScripting/SequencerScripting.uplugin
--- a/Engine/Plugins/MovieScene/SequencerScripting/SequencerScripting.uplugin
+++ b/Engine/Plugins/MovieScene/SequencerScripting/SequencerScripting.uplugin
@@ -22,6 +22,11 @@
 	"Modules": [
 		{
 			"Name": "SequencerScripting",
+			"Type": "Developer",
+			"LoadingPhase": "Default"
+		},
+		{
+			"Name": "SequencerScriptingEditor",
 			"Type": "Editor",
 			"LoadingPhase": "Default"
 		}
diff --git a/Engine/Plugins/MovieScene/SequencerScripting/Source/SequencerScripting/Private/SequencerTools.cpp b/Engine/Plugins/MovieScene/SequencerScripting/Source/SequencerScripting/Private/SequencerTools.cpp
deleted file mode 100644
diff --git a/Engine/Plugins/MovieScene/SequencerScripting/Source/SequencerScripting/Public/SequencerTools.h b/Engine/Plugins/MovieScene/SequencerScripting/Source/SequencerScripting/Public/SequencerTools.h
deleted file mode 100644
diff --git a/Engine/Plugins/MovieScene/SequencerScripting/Source/SequencerScripting/SequencerScripting.Build.cs b/Engine/Plugins/MovieScene/SequencerScripting/Source/SequencerScripting/SequencerScripting.Build.cs
--- a/Engine/Plugins/MovieScene/SequencerScripting/Source/SequencerScripting/SequencerScripting.Build.cs
+++ b/Engine/Plugins/MovieScene/SequencerScripting/Source/SequencerScripting/SequencerScripting.Build.cs
@@ -27,7 +27,6 @@ public class SequencerScripting : ModuleRules
 				"Core",
 				"TimeManagement",
 				"MovieScene",
-                "MovieSceneTools",
                 "MovieSceneTracks",
             }
 		);
@@ -38,11 +37,8 @@ public class SequencerScripting : ModuleRules
 				"CoreUObject",
 				"Engine",
 				"Kismet",
-				"PythonScriptPlugin",
 				"Slate",
 				"SlateCore",
-				"MovieSceneCaptureDialog",
-                "MovieSceneCapture",
                 "LevelSequence",
 			}
 		);
diff --git a/Engine/Plugins/MovieScene/SequencerScripting/Source/SequencerScriptingEditor/Private/SequencerScriptingEditor.cpp b/Engine/Plugins/MovieScene/SequencerScripting/Source/SequencerScriptingEditor/Private/SequencerScriptingEditor.cpp
new file mode 100644
--- /dev/null
+++ b/Engine/Plugins/MovieScene/SequencerScripting/Source/SequencerScriptingEditor/Private/SequencerScriptingEditor.cpp
@@ -0,0 +1,15 @@
+// Copyright 1998-2018 Epic Games, Inc. All Rights Reserved.
+
+#include "SequencerScriptingEditor.h"
+
+#define LOCTEXT_NAMESPACE "FSequencerScriptingEditorModule"
+
+void FSequencerScriptingEditorModule::StartupModule()
+{}
+
+void FSequencerScriptingEditorModule::ShutdownModule()
+{}
+
+#undef LOCTEXT_NAMESPACE
+
+IMPLEMENT_MODULE(FSequencerScriptingEditorModule, SequencerScriptingEditor)
\ No newline at end of file
diff --git a/Engine/Plugins/MovieScene/SequencerScripting/Source/SequencerScriptingEditor/Private/SequencerTools.cpp b/Engine/Plugins/MovieScene/SequencerScripting/Source/SequencerScriptingEditor/Private/SequencerTools.cpp
new file mode 100644
--- /dev/null
+++ b/Engine/Plugins/MovieScene/SequencerScripting/Source/SequencerScriptingEditor/Private/SequencerTools.cpp
@@ -0,0 +1,87 @@
+// Copyright 1998-2018 Epic Games, Inc. All Rights Reserved.
+
+#include "SequencerTools.h"
+#include "SequencerScriptingEditor.h"
+#include "MovieSceneCapture.h"
+#include "MovieSceneCaptureDialogModule.h"
+#include "AutomatedLevelSequenceCapture.h"
+#include "MovieSceneTimeHelpers.h"
+#include "UObject/Stack.h"
+#include "UObject/Package.h"
+
+#define LOCTEXT_NAMESPACE "SequencerTools"
+
+bool USequencerToolsFunctionLibrary::RenderMovie(UMovieSceneCapture* InCaptureSettings)
+{
+	IMovieSceneCaptureDialogModule& MovieSceneCaptureModule = FModuleManager::Get().LoadModuleChecked<IMovieSceneCaptureDialogModule>("MovieSceneCaptureDialog");
+	
+	// Because this comes from the Python/BP layer we need to soft-validate the state before we pass it onto functions that do a assert-based validation.
+	if (!InCaptureSettings)
+	{
+		FFrame::KismetExecutionMessage(TEXT("Cannot start Render Sequence to Movie with null capture settings."), ELogVerbosity::Error);
+		return false;
+	}
+
+	if (IsRenderingMovie())
+	{
+		FFrame::KismetExecutionMessage(TEXT("Capture already in progress."), ELogVerbosity::Error);
+		return false;
+	}
+
+	// If they're capturing a level sequence we'll do some additional checking as there are more parameters on the Automated Level Sequence capture.
+	UAutomatedLevelSequenceCapture* LevelSequenceCapture = Cast<UAutomatedLevelSequenceCapture>(InCaptureSettings);
+	if (LevelSequenceCapture)
+	{
+		if (!LevelSequenceCapture->LevelSequenceAsset.IsValid())
+		{
+			// UE_LOG(LogTemp, Warning, TEXT("No Level Sequence Asset specified in UAutomatedLevelSequenceCapture."));
+			FFrame::KismetExecutionMessage(TEXT("No Level Sequence Asset specified in UAutomatedLevelSequenceCapture."), ELogVerbosity::Error);
+			return false;
+		}
+
+		if (!LevelSequenceCapture->bUseCustomStartFrame && !LevelSequenceCapture->bUseCustomEndFrame)
+		{
+			// If they don't want to use a custom start/end frame we override the default values to be the length of the sequence, as the default is [0,1)
+			ULevelSequence* LevelSequence = Cast<ULevelSequence>(LevelSequenceCapture->LevelSequenceAsset.TryLoad());
+			if (!LevelSequence)
+			{
+				const FString ErrorMessage = FString::Printf(TEXT("Specified Level Sequence Asset failed to load. Specified Asset Path: %s"), *LevelSequenceCapture->LevelSequenceAsset.GetAssetPathString());
+				FFrame::KismetExecutionMessage(*ErrorMessage, ELogVerbosity::Error);
+				return false;
+			}
+
+			FFrameRate DisplayRate = LevelSequence->GetMovieScene()->GetDisplayRate();
+			FFrameRate TickResolution = LevelSequence->GetMovieScene()->GetTickResolution();
+			
+			LevelSequenceCapture->Settings.FrameRate = DisplayRate;
+			LevelSequenceCapture->Settings.bUseRelativeFrameNumbers = false;
+			TRange<FFrameNumber> Range = LevelSequence->GetMovieScene()->GetPlaybackRange();
+
+			FFrameNumber StartFrame = MovieScene::DiscreteInclusiveLower(Range);
+			FFrameNumber EndFrame = MovieScene::DiscreteExclusiveUpper(Range);
+
+			FFrameNumber RoundedStartFrame = FFrameRate::TransformTime(StartFrame, TickResolution, DisplayRate).CeilToFrame();
+			FFrameNumber RoundedEndFrame = FFrameRate::TransformTime(EndFrame, TickResolution, DisplayRate).CeilToFrame();
+
+			LevelSequenceCapture->CustomStartFrame = RoundedStartFrame;
+			LevelSequenceCapture->CustomEndFrame = RoundedEndFrame;
+		}
+	}
+
+	MovieSceneCaptureModule.StartCapture(InCaptureSettings);
+	return true;
+}
+
+void USequencerToolsFunctionLibrary::CancelMovieRender()
+{
+	IMovieSceneCaptureDialogModule& MovieSceneCaptureModule = FModuleManager::Get().LoadModuleChecked<IMovieSceneCaptureDialogModule>("MovieSceneCaptureDialog");
+	TSharedPtr<FMovieSceneCaptureBase> CurrentCapture = MovieSceneCaptureModule.GetCurrentCapture();
+	if (CurrentCapture.IsValid())
+	{
+		// We just invoke the capture's Cancel function. This will cause a shut-down of the capture (the same as the UI)
+		// which will invoke all of the necessary callbacks as well. We don't null out CurrentCapture because that is done
+		// as the result of its shutdown callbacks.
+		CurrentCapture->Cancel();
+	}
+}
+#undef LOCTEXT_NAMESPACE // "SequencerTools"
\ No newline at end of file
diff --git a/Engine/Plugins/MovieScene/SequencerScripting/Source/SequencerScriptingEditor/Public/SequencerScriptingEditor.h b/Engine/Plugins/MovieScene/SequencerScripting/Source/SequencerScriptingEditor/Public/SequencerScriptingEditor.h
new file mode 100644
--- /dev/null
+++ b/Engine/Plugins/MovieScene/SequencerScripting/Source/SequencerScriptingEditor/Public/SequencerScriptingEditor.h
@@ -0,0 +1,15 @@
+// Copyright 1998-2018 Epic Games, Inc. All Rights Reserved.
+
+#pragma once
+
+#include "CoreMinimal.h"
+#include "Modules/ModuleManager.h"
+
+class FSequencerScriptingEditorModule : public IModuleInterface
+{
+public:
+
+	/** IModuleInterface implementation */
+	virtual void StartupModule() override;
+	virtual void ShutdownModule() override;
+};
diff --git a/Engine/Plugins/MovieScene/SequencerScripting/Source/SequencerScriptingEditor/Public/SequencerTools.h b/Engine/Plugins/MovieScene/SequencerScripting/Source/SequencerScriptingEditor/Public/SequencerTools.h
new file mode 100644
--- /dev/null
+++ b/Engine/Plugins/MovieScene/SequencerScripting/Source/SequencerScriptingEditor/Public/SequencerTools.h
@@ -0,0 +1,44 @@
+// Copyright 1998-2018 Epic Games, Inc. All Rights Reserved.
+
+#pragma once
+
+#include "CoreMinimal.h"
+#include "Kismet/BlueprintFunctionLibrary.h"
+#include "Modules/ModuleManager.h"
+#include "MovieSceneCaptureDialogModule.h"
+#include "SequencerTools.generated.h"
+
+/** 
+ * This is a set of helper functions to access various parts of the Sequencer API via Python. Because Sequencer itself is not suitable for exposing, most functionality
+ * gets wrapped by UObjects that have an easier API to work with. This UObject provides access to these wrapper UObjects where needed. 
+ */
+UCLASS(Transient, meta=(ScriptName="SequencerTools"))
+class SEQUENCERSCRIPTINGEDITOR_API USequencerToolsFunctionLibrary : public UBlueprintFunctionLibrary
+{
+	GENERATED_BODY()
+public:
+
+	/**
+	* Attempts to render a sequence to movie based on the specified settings. This will automatically detect
+	* if we're rendering via a PIE instance or a new process based on the passed in settings. Will return false
+	* if the state is not valid (ie: null or missing required parameters, capture in progress, etc.), true otherwise.
+	*/
+	UFUNCTION(BlueprintCallable, Category = "Editor Scripting | Sequencer Tools")
+	static bool RenderMovie(class UMovieSceneCapture* InCaptureSettings);
+
+	/** 
+	* Returns if Render to Movie is currently in progress.
+	*/
+	UFUNCTION(BlueprintCallable, Category = "Editor Scripting | Sequencer Tools")
+	static bool IsRenderingMovie()
+	{
+		IMovieSceneCaptureDialogModule& MovieSceneCaptureModule = FModuleManager::Get().LoadModuleChecked<IMovieSceneCaptureDialogModule>("MovieSceneCaptureDialog");
+		return MovieSceneCaptureModule.GetCurrentCapture().IsValid();
+	}
+
+	/**
+	* Attempts to cancel an in-progress Render to Movie. Does nothing if there is no render in progress.
+	*/
+	UFUNCTION(BlueprintCallable, Category = "Editor Scripting | Sequencer Tools")
+	static void CancelMovieRender();
+};
\ No newline at end of file
diff --git a/Engine/Plugins/MovieScene/SequencerScripting/Source/SequencerScriptingEditor/SequencerScriptingEditor.Build.cs b/Engine/Plugins/MovieScene/SequencerScripting/Source/SequencerScriptingEditor/SequencerScriptingEditor.Build.cs
new file mode 100644
--- /dev/null
+++ b/Engine/Plugins/MovieScene/SequencerScripting/Source/SequencerScriptingEditor/SequencerScriptingEditor.Build.cs
@@ -0,0 +1,55 @@
+// Copyright 1998-2018 Epic Games, Inc. All Rights Reserved.
+
+using UnrealBuildTool;
+
+public class SequencerScriptingEditor : ModuleRules
+{
+	public SequencerScriptingEditor(ReadOnlyTargetRules Target) : base(Target)
+	{
+		PCHUsage = ModuleRules.PCHUsageMode.UseExplicitOrSharedPCHs;
+		
+		PublicIncludePaths.AddRange(
+			new string[] {
+            }
+		);
+
+		PrivateIncludePaths.AddRange(
+			new string[] {
+				"SequencerScriptingEditor/Private",
+            }
+		);
+
+		PublicDependencyModuleNames.AddRange(
+			new string[]
+			{
+				"Core",
+				"TimeManagement",
+				"MovieScene",
+                "MovieSceneTools",
+                "MovieSceneTracks",
+            }
+		);
+
+		PrivateDependencyModuleNames.AddRange(
+			new string[]
+			{
+				"CoreUObject",
+				"Engine",
+				"Kismet",
+				"PythonScriptPlugin",
+				"Slate",
+				"SlateCore",
+				"MovieSceneCaptureDialog",
+                "MovieSceneCapture",
+                "LevelSequence",
+			}
+		);
+
+		DynamicallyLoadedModuleNames.AddRange(
+			new string[]
+			{
+				
+			}
+		);
+	}
+}
