Sequencer Scripting: Add setters for playback range, frame rate, and tick resolution.

diff --git a/Engine/Plugins/MovieScene/SequencerScripting/Source/SequencerScripting/Private/ExtensionLibraries/MovieSceneSequenceExtensions.cpp b/Engine/Plugins/MovieScene/SequencerScripting/Source/SequencerScripting/Private/ExtensionLibraries/MovieSceneSequenceExtensions.cpp
--- a/Engine/Plugins/MovieScene/SequencerScripting/Source/SequencerScripting/Private/ExtensionLibraries/MovieSceneSequenceExtensions.cpp
+++ b/Engine/Plugins/MovieScene/SequencerScripting/Source/SequencerScripting/Private/ExtensionLibraries/MovieSceneSequenceExtensions.cpp
@@ -113,12 +113,30 @@ FFrameRate UMovieSceneSequenceExtensions::GetDisplayRate(UMovieSceneSequence* Se
 	return MovieScene ? MovieScene->GetDisplayRate() : FFrameRate();
 }
 
+void UMovieSceneSequenceExtensions::SetDisplayRate(UMovieSceneSequence* Sequence, FFrameRate DisplayRate)
+{
+	UMovieScene* MovieScene = Sequence->GetMovieScene();
+	if (MovieScene)
+	{
+		MovieScene->SetDisplayRate(DisplayRate);
+	}
+}
+
 FFrameRate UMovieSceneSequenceExtensions::GetTickResolution(UMovieSceneSequence* Sequence)
 {
 	UMovieScene* MovieScene = Sequence->GetMovieScene();
 	return MovieScene ? MovieScene->GetTickResolution() : FFrameRate();
 }
 
+void UMovieSceneSequenceExtensions::SetTickResolution(UMovieSceneSequence* Sequence, FFrameRate TickResolution)
+{
+	UMovieScene* MovieScene = Sequence->GetMovieScene();
+	if (MovieScene)
+	{
+		MovieScene->SetTickResolutionDirectly(TickResolution);
+	}
+}
+
 FSequencerScriptingRange UMovieSceneSequenceExtensions::MakeRange(UMovieSceneSequence* Sequence, int32 StartFrame, int32 Duration)
 {
 	FFrameRate FrameRate = GetDisplayRate(Sequence);
@@ -144,6 +162,15 @@ FSequencerScriptingRange UMovieSceneSequenceExtensions::GetPlaybackRange(UMovieS
 	}
 }
 
+void UMovieSceneSequenceExtensions::SetPlaybackRange(UMovieSceneSequence* Sequence, FSequencerScriptingRange PlaybackRange)
+{
+	UMovieScene* MovieScene = Sequence->GetMovieScene();
+	if (MovieScene)
+	{
+		MovieScene->SetPlaybackRange(PlaybackRange.ToNative(GetTickResolution(Sequence)));
+	}
+}
+
 FTimecode UMovieSceneSequenceExtensions::GetTimecodeSource(UMovieSceneSequence* Sequence)
 {
 	UMovieScene* MovieScene = Sequence->GetMovieScene();
diff --git a/Engine/Plugins/MovieScene/SequencerScripting/Source/SequencerScripting/Public/ExtensionLibraries/MovieSceneSequenceExtensions.h b/Engine/Plugins/MovieScene/SequencerScripting/Source/SequencerScripting/Public/ExtensionLibraries/MovieSceneSequenceExtensions.h
--- a/Engine/Plugins/MovieScene/SequencerScripting/Source/SequencerScripting/Public/ExtensionLibraries/MovieSceneSequenceExtensions.h
+++ b/Engine/Plugins/MovieScene/SequencerScripting/Source/SequencerScripting/Public/ExtensionLibraries/MovieSceneSequenceExtensions.h
@@ -70,7 +70,7 @@ public:
 	static UMovieSceneTrack* AddMasterTrack(UMovieSceneSequence* Sequence, TSubclassOf<UMovieSceneTrack> TrackType);
 
 	/**
-	 * Get's this sequence's display rate
+	 * Gets this sequence's display rate
 	 *
 	 * @param Sequence        The sequence to use
 	 * @return The display rate that this sequence is displayed as
@@ -79,7 +79,16 @@ public:
 	static FFrameRate GetDisplayRate(UMovieSceneSequence* Sequence);
 
 	/**
-	 * Get's this sequence's tick resolution
+	 * Sets this sequence's display rate
+	 *
+	 * @param Sequence        The sequence to use
+	 * @param DisplayRate The display rate that this sequence is displayed as
+	 */
+	UFUNCTION(BlueprintCallable, Category="Sequence", meta=(ScriptMethod))
+	static void SetDisplayRate(UMovieSceneSequence* Sequence, FFrameRate DisplayRate);
+
+	/**
+	 * Gets this sequence's tick resolution
 	 *
 	 * @param Sequence        The sequence to use
 	 * @return The tick resolution of the sequence, defining the smallest unit of time representable on this sequence
@@ -87,6 +96,15 @@ public:
 	UFUNCTION(BlueprintCallable, Category="Sequence", meta=(ScriptMethod))
 	static FFrameRate GetTickResolution(UMovieSceneSequence* Sequence);
 
+	/**
+	 * Sets this sequence's tick resolution
+	 *
+	 * @param Sequence        The sequence to use
+	 * @param TickResolution The tick resolution of the sequence, defining the smallest unit of time representable on this sequence
+	 */
+	UFUNCTION(BlueprintCallable, Category="Sequence", meta=(ScriptMethod))
+	static void SetTickResolution(UMovieSceneSequence* Sequence, FFrameRate TickResolution);
+
 	/**
 	 * Make a new range for this sequence in its display rate
 	 *
@@ -119,11 +137,20 @@ public:
 	static FSequencerScriptingRange GetPlaybackRange(UMovieSceneSequence* Sequence);
 
 	/**
-	* Get the timecode source of this sequence
-	*
-	* @param Sequence        The sequence within which to get the timecode source
-	* @return Timecode source of this sequence
-	*/
+	 * Set playback range of this sequence
+	 *
+	 * @param Sequence        The sequence within which to set the playback range
+	 * @param PlaybackRange Playback range of this sequence
+	 */
+	UFUNCTION(BlueprintCallable, Category="Sequence", meta=(ScriptMethod))
+	static void SetPlaybackRange(UMovieSceneSequence* Sequence, FSequencerScriptingRange PlaybackRange);
+
+	/**
+	 * Get the timecode source of this sequence
+	 *
+	 * @param Sequence        The sequence within which to get the timecode source
+	 * @return Timecode source of this sequence
+	 */
 	UFUNCTION(BlueprintCallable, Category = "Sequence", meta = (ScriptMethod))
 	static FTimecode GetTimecodeSource(UMovieSceneSequence* Sequence);
 
