Sequencer: Call GetMovieScene() which checks for Sequence validity

diff --git a/Engine/Plugins/MovieScene/SequencerScripting/Source/SequencerScripting/Private/ExtensionLibraries/MovieSceneSequenceExtensions.cpp b/Engine/Plugins/MovieScene/SequencerScripting/Source/SequencerScripting/Private/ExtensionLibraries/MovieSceneSequenceExtensions.cpp
--- a/Engine/Plugins/MovieScene/SequencerScripting/Source/SequencerScripting/Private/ExtensionLibraries/MovieSceneSequenceExtensions.cpp
+++ b/Engine/Plugins/MovieScene/SequencerScripting/Source/SequencerScripting/Private/ExtensionLibraries/MovieSceneSequenceExtensions.cpp
@@ -24,14 +24,14 @@ TArray<UMovieSceneTrack*> UMovieSceneSequenceExtensions::FilterTracks(TArrayView
 
 UMovieScene* UMovieSceneSequenceExtensions::GetMovieScene(UMovieSceneSequence* Sequence)
 {
-	return Sequence->GetMovieScene();
+	return Sequence ? Sequence->GetMovieScene() : nullptr;
 }
 
 TArray<UMovieSceneTrack*> UMovieSceneSequenceExtensions::GetMasterTracks(UMovieSceneSequence* Sequence)
 {
 	TArray<UMovieSceneTrack*> Tracks;
 
-	UMovieScene* MovieScene = Sequence->GetMovieScene();
+	UMovieScene* MovieScene = GetMovieScene(Sequence);
 	if (MovieScene)
 	{
 		Tracks = MovieScene->GetMasterTracks();
@@ -47,7 +47,7 @@ TArray<UMovieSceneTrack*> UMovieSceneSequenceExtensions::GetMasterTracks(UMovieS
 
 TArray<UMovieSceneTrack*> UMovieSceneSequenceExtensions::FindMasterTracksByType(UMovieSceneSequence* Sequence, TSubclassOf<UMovieSceneTrack> TrackType)
 {
-	UMovieScene* MovieScene   = Sequence->GetMovieScene();
+	UMovieScene* MovieScene   = GetMovieScene(Sequence);
 	UClass*      DesiredClass = TrackType.Get();
 
 	if (MovieScene && DesiredClass)
@@ -70,7 +70,7 @@ TArray<UMovieSceneTrack*> UMovieSceneSequenceExtensions::FindMasterTracksByType(
 
 TArray<UMovieSceneTrack*> UMovieSceneSequenceExtensions::FindMasterTracksByExactType(UMovieSceneSequence* Sequence, TSubclassOf<UMovieSceneTrack> TrackType)
 {
-	UMovieScene* MovieScene   = Sequence->GetMovieScene();
+	UMovieScene* MovieScene   = GetMovieScene(Sequence);
 	UClass*      DesiredClass = TrackType.Get();
 
 	if (MovieScene && DesiredClass)
@@ -94,7 +94,7 @@ TArray<UMovieSceneTrack*> UMovieSceneSequenceExtensions::FindMasterTracksByExact
 UMovieSceneTrack* UMovieSceneSequenceExtensions::AddMasterTrack(UMovieSceneSequence* Sequence, TSubclassOf<UMovieSceneTrack> TrackType)
 {
 	// @todo: sequencer-python: master track type compatibility with sequence. Currently that's really only loosely defined by track editors, which is not sufficient here.
-	UMovieScene* MovieScene = Sequence->GetMovieScene();
+	UMovieScene* MovieScene = GetMovieScene(Sequence);
 	if (MovieScene)
 	{
 		UMovieSceneTrack* NewTrack = MovieScene->AddMasterTrack(TrackType);
@@ -109,13 +109,13 @@ UMovieSceneTrack* UMovieSceneSequenceExtensions::AddMasterTrack(UMovieSceneSeque
 
 FFrameRate UMovieSceneSequenceExtensions::GetDisplayRate(UMovieSceneSequence* Sequence)
 {
-	UMovieScene* MovieScene = Sequence->GetMovieScene();
+	UMovieScene* MovieScene = GetMovieScene(Sequence);
 	return MovieScene ? MovieScene->GetDisplayRate() : FFrameRate();
 }
 
 void UMovieSceneSequenceExtensions::SetDisplayRate(UMovieSceneSequence* Sequence, FFrameRate DisplayRate)
 {
-	UMovieScene* MovieScene = Sequence->GetMovieScene();
+	UMovieScene* MovieScene = GetMovieScene(Sequence);
 	if (MovieScene)
 	{
 		MovieScene->SetDisplayRate(DisplayRate);
@@ -124,13 +124,13 @@ void UMovieSceneSequenceExtensions::SetDisplayRate(UMovieSceneSequence* Sequence
 
 FFrameRate UMovieSceneSequenceExtensions::GetTickResolution(UMovieSceneSequence* Sequence)
 {
-	UMovieScene* MovieScene = Sequence->GetMovieScene();
+	UMovieScene* MovieScene = GetMovieScene(Sequence);
 	return MovieScene ? MovieScene->GetTickResolution() : FFrameRate();
 }
 
 void UMovieSceneSequenceExtensions::SetTickResolution(UMovieSceneSequence* Sequence, FFrameRate TickResolution)
 {
-	UMovieScene* MovieScene = Sequence->GetMovieScene();
+	UMovieScene* MovieScene = GetMovieScene(Sequence);
 	if (MovieScene)
 	{
 		MovieScene->SetTickResolutionDirectly(TickResolution);
@@ -151,7 +151,7 @@ FSequencerScriptingRange UMovieSceneSequenceExtensions::MakeRangeSeconds(UMovieS
 
 FSequencerScriptingRange UMovieSceneSequenceExtensions::GetPlaybackRange(UMovieSceneSequence* Sequence)
 {
-	UMovieScene* MovieScene = Sequence->GetMovieScene();
+	UMovieScene* MovieScene = GetMovieScene(Sequence);
 	if (MovieScene)
 	{
 		return FSequencerScriptingRange::FromNative(MovieScene->GetPlaybackRange(), GetTickResolution(Sequence));
@@ -164,7 +164,7 @@ FSequencerScriptingRange UMovieSceneSequenceExtensions::GetPlaybackRange(UMovieS
 
 void UMovieSceneSequenceExtensions::SetPlaybackRange(UMovieSceneSequence* Sequence, FSequencerScriptingRange PlaybackRange)
 {
-	UMovieScene* MovieScene = Sequence->GetMovieScene();
+	UMovieScene* MovieScene = GetMovieScene(Sequence);
 	if (MovieScene)
 	{
 		MovieScene->SetPlaybackRange(PlaybackRange.ToNative(GetTickResolution(Sequence)));
@@ -173,7 +173,7 @@ void UMovieSceneSequenceExtensions::SetPlaybackRange(UMovieSceneSequence* Sequen
 
 FTimecode UMovieSceneSequenceExtensions::GetTimecodeSource(UMovieSceneSequence* Sequence)
 {
-	UMovieScene* MovieScene = Sequence->GetMovieScene();
+	UMovieScene* MovieScene = GetMovieScene(Sequence);
 	if (MovieScene)
 	{
 		return MovieScene->TimecodeSource.Timecode;
@@ -186,7 +186,7 @@ FTimecode UMovieSceneSequenceExtensions::GetTimecodeSource(UMovieSceneSequence*
 
 FSequencerBindingProxy UMovieSceneSequenceExtensions::FindBindingByName(UMovieSceneSequence* Sequence, FString Name)
 {
-	UMovieScene* MovieScene = Sequence->GetMovieScene();
+	UMovieScene* MovieScene = GetMovieScene(Sequence);
 	if (MovieScene)
 	{
 		const FMovieSceneBinding* Binding = Algo::FindBy(MovieScene->GetBindings(), Name, &FMovieSceneBinding::GetName);
@@ -202,7 +202,7 @@ TArray<FSequencerBindingProxy> UMovieSceneSequenceExtensions::GetBindings(UMovie
 {
 	TArray<FSequencerBindingProxy> AllBindings;
 
-	UMovieScene* MovieScene = Sequence->GetMovieScene();
+	UMovieScene* MovieScene = GetMovieScene(Sequence);
 	if (MovieScene)
 	{
 		for (const FMovieSceneBinding& Binding : MovieScene->GetBindings())
@@ -218,7 +218,7 @@ TArray<FSequencerBindingProxy> UMovieSceneSequenceExtensions::GetSpawnables(UMov
 {
 	TArray<FSequencerBindingProxy> AllSpawnables;
 
-	UMovieScene* MovieScene = Sequence->GetMovieScene();
+	UMovieScene* MovieScene = GetMovieScene(Sequence);
 	if (MovieScene)
 	{
 		int32 Count = MovieScene->GetSpawnableCount();
@@ -236,7 +236,7 @@ TArray<FSequencerBindingProxy> UMovieSceneSequenceExtensions::GetPossessables(UM
 {
 	TArray<FSequencerBindingProxy> AllPossessables;
 
-	UMovieScene* MovieScene = Sequence->GetMovieScene();
+	UMovieScene* MovieScene = GetMovieScene(Sequence);
 	if (MovieScene)
 	{
 		int32 Count = MovieScene->GetPossessableCount();
