Backout SetPlaybackRange. It will be rewritten as SetPlaybackStart, SetPlaybackEnd.

diff --git a/Engine/Plugins/MovieScene/SequencerScripting/Source/SequencerScripting/Private/ExtensionLibraries/MovieSceneSequenceExtensions.cpp b/Engine/Plugins/MovieScene/SequencerScripting/Source/SequencerScripting/Private/ExtensionLibraries/MovieSceneSequenceExtensions.cpp
--- a/Engine/Plugins/MovieScene/SequencerScripting/Source/SequencerScripting/Private/ExtensionLibraries/MovieSceneSequenceExtensions.cpp
+++ b/Engine/Plugins/MovieScene/SequencerScripting/Source/SequencerScripting/Private/ExtensionLibraries/MovieSceneSequenceExtensions.cpp
@@ -162,15 +162,6 @@ FSequencerScriptingRange UMovieSceneSequenceExtensions::GetPlaybackRange(UMovieS
 	}
 }
 
-void UMovieSceneSequenceExtensions::SetPlaybackRange(UMovieSceneSequence* Sequence, FSequencerScriptingRange PlaybackRange)
-{
-	UMovieScene* MovieScene = GetMovieScene(Sequence);
-	if (MovieScene)
-	{
-		MovieScene->SetPlaybackRange(PlaybackRange.ToNative(GetTickResolution(Sequence)));
-	}
-}
-
 FTimecode UMovieSceneSequenceExtensions::GetTimecodeSource(UMovieSceneSequence* Sequence)
 {
 	UMovieScene* MovieScene = GetMovieScene(Sequence);
diff --git a/Engine/Plugins/MovieScene/SequencerScripting/Source/SequencerScripting/Public/ExtensionLibraries/MovieSceneSequenceExtensions.h b/Engine/Plugins/MovieScene/SequencerScripting/Source/SequencerScripting/Public/ExtensionLibraries/MovieSceneSequenceExtensions.h
--- a/Engine/Plugins/MovieScene/SequencerScripting/Source/SequencerScripting/Public/ExtensionLibraries/MovieSceneSequenceExtensions.h
+++ b/Engine/Plugins/MovieScene/SequencerScripting/Source/SequencerScripting/Public/ExtensionLibraries/MovieSceneSequenceExtensions.h
@@ -136,15 +136,6 @@ public:
 	UFUNCTION(BlueprintCallable, Category="Sequence", meta=(ScriptMethod))
 	static FSequencerScriptingRange GetPlaybackRange(UMovieSceneSequence* Sequence);
 
-	/**
-	 * Set playback range of this sequence
-	 *
-	 * @param Sequence        The sequence within which to set the playback range
-	 * @param PlaybackRange Playback range of this sequence
-	 */
-	UFUNCTION(BlueprintCallable, Category="Sequence", meta=(ScriptMethod))
-	static void SetPlaybackRange(UMovieSceneSequence* Sequence, FSequencerScriptingRange PlaybackRange);
-
 	/**
 	 * Get the timecode source of this sequence
 	 *
