Sequencer: More python extensions - set parent, set property name and path, set start/end frame bounded

sequencer_examples.py - Added a current focal length properry track for cameras, modified to use set_start_frame and set_end_frame.

diff --git a/Engine/Plugins/MovieScene/SequencerScripting/Source/SequencerScripting/Private/ExtensionLibraries/MovieSceneBindingExtensions.cpp b/Engine/Plugins/MovieScene/SequencerScripting/Source/SequencerScripting/Private/ExtensionLibraries/MovieSceneBindingExtensions.cpp
--- a/Engine/Plugins/MovieScene/SequencerScripting/Source/SequencerScripting/Private/ExtensionLibraries/MovieSceneBindingExtensions.cpp
+++ b/Engine/Plugins/MovieScene/SequencerScripting/Source/SequencerScripting/Private/ExtensionLibraries/MovieSceneBindingExtensions.cpp
@@ -216,3 +216,17 @@ FSequencerBindingProxy UMovieSceneBindingExtensions::GetParent(const FSequencerB
 	}
 	return FSequencerBindingProxy();
 }
+
+
+void UMovieSceneBindingExtensions::SetParent(const FSequencerBindingProxy& InBinding, const FSequencerBindingProxy& InParentBinding)
+{
+	UMovieScene* MovieScene = InBinding.GetMovieScene();
+	if (MovieScene)
+	{
+		FMovieScenePossessable* Possessable = MovieScene->FindPossessable(InBinding.BindingID);
+		if (Possessable)
+		{
+			Possessable->SetParent(InParentBinding.BindingID);
+		}
+	}
+}
diff --git a/Engine/Plugins/MovieScene/SequencerScripting/Source/SequencerScripting/Private/ExtensionLibraries/MovieScenePropertyTrackExtensions.cpp b/Engine/Plugins/MovieScene/SequencerScripting/Source/SequencerScripting/Private/ExtensionLibraries/MovieScenePropertyTrackExtensions.cpp
--- a/Engine/Plugins/MovieScene/SequencerScripting/Source/SequencerScripting/Private/ExtensionLibraries/MovieScenePropertyTrackExtensions.cpp
+++ b/Engine/Plugins/MovieScene/SequencerScripting/Source/SequencerScripting/Private/ExtensionLibraries/MovieScenePropertyTrackExtensions.cpp
@@ -3,6 +3,11 @@
 #include "ExtensionLibraries/MovieScenePropertyTrackExtensions.h"
 #include "Tracks/MovieScenePropertyTrack.h"
 
+void UMovieScenePropertyTrackExtensions::SetPropertyNameAndPath(UMovieScenePropertyTrack* Track, const FName& InPropertyName, const FString& InPropertyPath)
+{
+	Track->SetPropertyNameAndPath(InPropertyName, InPropertyPath);
+}
+
 FName UMovieScenePropertyTrackExtensions::GetPropertyName(UMovieScenePropertyTrack* Track)
 {
 	return Track->GetPropertyName();
diff --git a/Engine/Plugins/MovieScene/SequencerScripting/Source/SequencerScripting/Private/ExtensionLibraries/MovieSceneSectionExtensions.cpp b/Engine/Plugins/MovieScene/SequencerScripting/Source/SequencerScripting/Private/ExtensionLibraries/MovieSceneSectionExtensions.cpp
--- a/Engine/Plugins/MovieScene/SequencerScripting/Source/SequencerScripting/Private/ExtensionLibraries/MovieSceneSectionExtensions.cpp
+++ b/Engine/Plugins/MovieScene/SequencerScripting/Source/SequencerScripting/Private/ExtensionLibraries/MovieSceneSectionExtensions.cpp
@@ -110,6 +110,28 @@ void UMovieSceneSectionExtensions::SetStartFrameSeconds(UMovieSceneSection* Sect
 	}
 }
 
+void UMovieSceneSectionExtensions::SetStartFrameBounded(UMovieSceneSection* Section, bool bIsBounded)
+{
+	UMovieScene* MovieScene = Section->GetTypedOuter<UMovieScene>();
+	if (MovieScene)
+	{
+		if (bIsBounded)
+		{
+			int32 NewFrameNumber = 0;
+			if (!MovieScene->GetPlaybackRange().GetLowerBound().IsOpen())
+			{
+				NewFrameNumber = MovieScene->GetPlaybackRange().GetLowerBoundValue().Value;
+			}
+
+			Section->SectionRange.Value.SetLowerBound(TRangeBound<FFrameNumber>(FFrameNumber(NewFrameNumber)));
+		}
+		else
+		{
+			Section->SectionRange.Value.SetLowerBound(TRangeBound<FFrameNumber>());
+		}
+	}
+}
+
 void UMovieSceneSectionExtensions::SetEndFrame(UMovieSceneSection* Section, int32 EndFrame)
 {
 	UMovieScene* MovieScene = Section->GetTypedOuter<UMovieScene>();
@@ -130,6 +152,28 @@ void UMovieSceneSectionExtensions::SetEndFrameSeconds(UMovieSceneSection* Sectio
 	}
 }
 
+void UMovieSceneSectionExtensions::SetEndFrameBounded(UMovieSceneSection* Section, bool bIsBounded)
+{
+	UMovieScene* MovieScene = Section->GetTypedOuter<UMovieScene>();
+	if (MovieScene)
+	{
+		if (bIsBounded)
+		{
+			int32 NewFrameNumber = 0;
+			if (!MovieScene->GetPlaybackRange().GetUpperBound().IsOpen())
+			{
+				NewFrameNumber = MovieScene->GetPlaybackRange().GetUpperBoundValue().Value;
+			}
+
+			Section->SectionRange.Value.SetUpperBound(TRangeBound<FFrameNumber>(FFrameNumber(NewFrameNumber)));
+		}
+		else
+		{
+			Section->SectionRange.Value.SetUpperBound(TRangeBound<FFrameNumber>());
+		}
+	}
+}
+
 template<typename ChannelType, typename ScriptingChannelType>
 void GetScriptingChannelsForChannel(FMovieSceneChannelProxy& ChannelProxy, TWeakObjectPtr<UMovieSceneSequence> Sequence, TArray<UMovieSceneScriptingChannel*>& OutChannels)
 {
diff --git a/Engine/Plugins/MovieScene/SequencerScripting/Source/SequencerScripting/Public/ExtensionLibraries/MovieSceneBindingExtensions.h b/Engine/Plugins/MovieScene/SequencerScripting/Source/SequencerScripting/Public/ExtensionLibraries/MovieSceneBindingExtensions.h
--- a/Engine/Plugins/MovieScene/SequencerScripting/Source/SequencerScripting/Public/ExtensionLibraries/MovieSceneBindingExtensions.h
+++ b/Engine/Plugins/MovieScene/SequencerScripting/Source/SequencerScripting/Public/ExtensionLibraries/MovieSceneBindingExtensions.h
@@ -144,4 +144,13 @@ public:
 	*/
 	UFUNCTION(BlueprintCallable, Category=Sequence, meta=(ScriptMethod))
 	static FSequencerBindingProxy GetParent(const FSequencerBindingProxy& InBinding);
+
+	/**
+	 * Set the parent to this binding
+	 *
+	 * @param InBinding     The binding to set 
+	 * @param InParentBinding     The parent to set the InBinding to
+	 */
+	UFUNCTION(BlueprintCallable, Category = Sequence, meta = (ScriptMethod))
+	static void SetParent(const FSequencerBindingProxy& InBinding, const FSequencerBindingProxy& InParentBinding);
 };
\ No newline at end of file
diff --git a/Engine/Plugins/MovieScene/SequencerScripting/Source/SequencerScripting/Public/ExtensionLibraries/MovieScenePropertyTrackExtensions.h b/Engine/Plugins/MovieScene/SequencerScripting/Source/SequencerScripting/Public/ExtensionLibraries/MovieScenePropertyTrackExtensions.h
--- a/Engine/Plugins/MovieScene/SequencerScripting/Source/SequencerScripting/Public/ExtensionLibraries/MovieScenePropertyTrackExtensions.h
+++ b/Engine/Plugins/MovieScene/SequencerScripting/Source/SequencerScripting/Public/ExtensionLibraries/MovieScenePropertyTrackExtensions.h
@@ -21,6 +21,16 @@ class UMovieScenePropertyTrackExtensions : public UBlueprintFunctionLibrary
 
 public:
 
+	/**
+	 * Set this track's property name and path
+	 *
+	 * @param Track        The track to use
+	 * @param InPropertyName The property name
+	 * @param InPropertyPath The property path
+	 */
+	UFUNCTION(BlueprintCallable, Category = "Track", meta = (ScriptMethod))
+	static void SetPropertyNameAndPath(UMovieScenePropertyTrack* Track, const FName& InPropertyName, const FString& InPropertyPath);
+
 	/**
 	 * Get this track's property name
 	 *
diff --git a/Engine/Plugins/MovieScene/SequencerScripting/Source/SequencerScripting/Public/ExtensionLibraries/MovieSceneSectionExtensions.h b/Engine/Plugins/MovieScene/SequencerScripting/Source/SequencerScripting/Public/ExtensionLibraries/MovieSceneSectionExtensions.h
--- a/Engine/Plugins/MovieScene/SequencerScripting/Source/SequencerScripting/Public/ExtensionLibraries/MovieSceneSectionExtensions.h
+++ b/Engine/Plugins/MovieScene/SequencerScripting/Source/SequencerScripting/Public/ExtensionLibraries/MovieSceneSectionExtensions.h
@@ -81,6 +81,15 @@ public:
 	UFUNCTION(BlueprintCallable, Category = "Section", meta = (ScriptMethod))
 	static void SetStartFrameSeconds(UMovieSceneSection* Section, float StartTime);
 
+	/**
+	 * Set start frame bounded
+	 *
+	 * @param Section        The section to set whether the start frame is bounded or not
+	 * @param IsBounded The desired bounded state of the start frame
+	 */
+	UFUNCTION(BlueprintCallable, Category = "Section", meta = (ScriptMethod))
+	static void SetStartFrameBounded(UMovieSceneSection* Section, bool bIsBounded);
+
 	/**
 	 * Set end frame
 	 *
@@ -99,6 +108,15 @@ public:
 	UFUNCTION(BlueprintCallable, Category = "Section", meta = (ScriptMethod))
 	static void SetEndFrameSeconds(UMovieSceneSection* Section, float EndTime);
 
+	/**
+     * Set end frame bounded
+	 *
+	 * @param Section        The section to set whether the end frame is bounded or not
+	 * @param IsBounded The desired bounded state of the end frame
+	 */
+	UFUNCTION(BlueprintCallable, Category = "Section", meta = (ScriptMethod))
+	static void SetEndFrameBounded(UMovieSceneSection* Section, bool bIsBounded);
+
 	/**
 	* Find all channels that belong to the specified UMovieSceneSection. Some sections have many channels (such as
 	* Transforms containing 9 float channels to represent Translation/Rotation/Scale), and a section may have mixed
