Sequencer: Added back SetRange but with frame numbers in display resolution. Added SetRangeSeconds

diff --git a/Engine/Plugins/MovieScene/SequencerScripting/Source/SequencerScripting/Private/ExtensionLibraries/MovieSceneSectionExtensions.cpp b/Engine/Plugins/MovieScene/SequencerScripting/Source/SequencerScripting/Private/ExtensionLibraries/MovieSceneSectionExtensions.cpp
--- a/Engine/Plugins/MovieScene/SequencerScripting/Source/SequencerScripting/Private/ExtensionLibraries/MovieSceneSectionExtensions.cpp
+++ b/Engine/Plugins/MovieScene/SequencerScripting/Source/SequencerScripting/Private/ExtensionLibraries/MovieSceneSectionExtensions.cpp
@@ -75,18 +75,47 @@ float UMovieSceneSectionExtensions::GetEndFrameSeconds(UMovieSceneSection* Secti
 	return -1.f;
 }
 
-void UMovieSceneSectionExtensions::SetRange(UMovieSceneSection* Section, const FSequencerScriptingRange& InRange)
+void UMovieSceneSectionExtensions::SetRange(UMovieSceneSection* Section, int32 StartFrame, int32 EndFrame)
 {
 	UMovieScene* MovieScene = Section->GetTypedOuter<UMovieScene>();
-	TRange<FFrameNumber> NewRange = InRange.ToNative(MovieScene->GetTickResolution());
-
-	if (NewRange.GetLowerBound().IsOpen() || NewRange.GetUpperBound().IsOpen() || NewRange.GetLowerBoundValue() <= NewRange.GetUpperBoundValue())
+	if (MovieScene)
 	{
-		Section->SetRange(NewRange);
+		FFrameRate DisplayRate = MovieScene->GetDisplayRate();
+
+		TRange<FFrameNumber> NewRange;
+		NewRange.SetLowerBound(TRangeBound<FFrameNumber>::Inclusive(ConvertFrameTime(StartFrame, DisplayRate, MovieScene->GetTickResolution()).FrameNumber));
+		NewRange.SetUpperBound(TRangeBound<FFrameNumber>::Inclusive(ConvertFrameTime(EndFrame, DisplayRate, MovieScene->GetTickResolution()).FrameNumber));
+
+		if (NewRange.GetLowerBound().IsOpen() || NewRange.GetUpperBound().IsOpen() || NewRange.GetLowerBoundValue() <= NewRange.GetUpperBoundValue())
+		{
+			Section->SetRange(NewRange);
+		}
+		else
+		{
+			UE_LOG(LogMovieScene, Error, TEXT("Invalid range specified"));
+		}
 	}
-	else
+}
+
+void UMovieSceneSectionExtensions::SetRangeSeconds(UMovieSceneSection* Section, float StartTime, float EndTime)
+{
+	UMovieScene* MovieScene = Section->GetTypedOuter<UMovieScene>();
+	if (MovieScene)
 	{
-		UE_LOG(LogMovieScene, Error, TEXT("Invalid range specified"));
+		FFrameRate DisplayRate = MovieScene->GetDisplayRate();
+
+		TRange<FFrameNumber> NewRange;
+		NewRange.SetLowerBound((StartTime * MovieScene->GetTickResolution()).RoundToFrame());
+		NewRange.SetUpperBound((EndTime * MovieScene->GetTickResolution()).RoundToFrame());
+
+		if (NewRange.GetLowerBound().IsOpen() || NewRange.GetUpperBound().IsOpen() || NewRange.GetLowerBoundValue() <= NewRange.GetUpperBoundValue())
+		{
+			Section->SetRange(NewRange);
+		}
+		else
+		{
+			UE_LOG(LogMovieScene, Error, TEXT("Invalid range specified"));
+		}
 	}
 }
 
diff --git a/Engine/Plugins/MovieScene/SequencerScripting/Source/SequencerScripting/Public/ExtensionLibraries/MovieSceneSectionExtensions.h b/Engine/Plugins/MovieScene/SequencerScripting/Source/SequencerScripting/Public/ExtensionLibraries/MovieSceneSectionExtensions.h
--- a/Engine/Plugins/MovieScene/SequencerScripting/Source/SequencerScripting/Public/ExtensionLibraries/MovieSceneSectionExtensions.h
+++ b/Engine/Plugins/MovieScene/SequencerScripting/Source/SequencerScripting/Public/ExtensionLibraries/MovieSceneSectionExtensions.h
@@ -60,8 +60,25 @@ public:
 	UFUNCTION(BlueprintCallable, Category = "Section", meta = (ScriptMethod))
 	static float GetEndFrameSeconds(UMovieSceneSection* Section);
 
-	UE_DEPRECATED(4.22, "Please use SetStartFrame and SetEndFrame instead.")
-	static void SetRange(UMovieSceneSection* Section, const FSequencerScriptingRange& Range);
+	/**
+	 * Set range
+	 *
+	 * @param Section        The section within which to set the range
+	 * @param StartFrame The desired start frame for this section
+	 * @param EndFrame The desired end frame for this section
+	 */
+	UFUNCTION(BlueprintCallable, Category = "Section", meta = (ScriptMethod))
+	static void SetRange(UMovieSceneSection* Section, int32 StartFrame, int32 EndFrame);
+
+	/**
+	 * Set range in seconds
+	 *
+	 * @param Section        The section within which to set the range
+	 * @param StartTime The desired start frame for this section
+	 * @param EndTime The desired end frame for this section
+	 */
+	UFUNCTION(BlueprintCallable, Category = "Section", meta = (ScriptMethod))
+	static void SetRangeSeconds(UMovieSceneSection* Section, float StartTime, float EndTime);
 
 	/**
 	 * Set start frame
