Sequencer: Expose Set/GetNumChannelsUsed.

Added a dummy 2 vector channel track to sequencer_examples.py

diff --git a/Engine/Plugins/MovieScene/SequencerScripting/Source/SequencerScripting/Private/ExtensionLibraries/MovieSceneVectorTrackExtensions.cpp b/Engine/Plugins/MovieScene/SequencerScripting/Source/SequencerScripting/Private/ExtensionLibraries/MovieSceneVectorTrackExtensions.cpp
new file mode 100644
--- /dev/null
+++ b/Engine/Plugins/MovieScene/SequencerScripting/Source/SequencerScripting/Private/ExtensionLibraries/MovieSceneVectorTrackExtensions.cpp
@@ -0,0 +1,15 @@
+// Copyright 1998-2019 Epic Games, Inc. All Rights Reserved.
+
+#include "ExtensionLibraries/MovieSceneVectorTrackExtensions.h"
+#include "Tracks/MovieSceneVectorTrack.h"
+
+void UMovieSceneVectorTrackExtensions::SetNumChannelsUsed(UMovieSceneVectorTrack* Track, int32 InNumChannelsUsed)
+{
+	Track->SetNumChannelsUsed(InNumChannelsUsed);
+}
+
+
+int32 UMovieSceneVectorTrackExtensions::GetNumChannelsUsed(UMovieSceneVectorTrack* Track)
+{
+	return Track->GetNumChannelsUsed();
+}
diff --git a/Engine/Plugins/MovieScene/SequencerScripting/Source/SequencerScripting/Public/ExtensionLibraries/MovieSceneVectorTrackExtensions.h b/Engine/Plugins/MovieScene/SequencerScripting/Source/SequencerScripting/Public/ExtensionLibraries/MovieSceneVectorTrackExtensions.h
new file mode 100644
--- /dev/null
+++ b/Engine/Plugins/MovieScene/SequencerScripting/Source/SequencerScripting/Public/ExtensionLibraries/MovieSceneVectorTrackExtensions.h
@@ -0,0 +1,43 @@
+// Copyright 1998-2019 Epic Games, Inc. All Rights Reserved.
+
+#pragma once
+
+#include "CoreTypes.h"
+#include "UObject/ObjectMacros.h"
+#include "Kismet/BlueprintFunctionLibrary.h"
+
+#include "MovieSceneVectorTrackExtensions.generated.h"
+
+
+class UMovieSceneVectorTrack;
+
+/**
+ * Function library containing methods that should be hoisted onto UMovieSceneVectorTrack for scripting
+ */
+UCLASS()
+class UMovieSceneVectorTrackExtensions : public UBlueprintFunctionLibrary
+{
+	GENERATED_BODY()
+
+public:
+
+	/**
+	 * Set the number of channels used for this track
+	 *
+	 * @param Track        The track to set
+	 * @param InNumChannelsUsed The number of channels to use for this track
+	 */
+	UFUNCTION(BlueprintCallable, Category = "Track", meta = (ScriptMethod))
+	static void SetNumChannelsUsed(UMovieSceneVectorTrack* Track, int32 InNumChannelsUsed);
+
+	/**
+	 * Get the number of channels used for this track
+	 *
+	 * @param Track        The track to query for the number of channels used
+	 * @return The number of channels used for this track
+	 */
+	UFUNCTION(BlueprintCallable, Category = "Track", meta = (ScriptMethod))
+	static int32 GetNumChannelsUsed(UMovieSceneVectorTrack* Track);
+
+
+};
