Sequencer: Added GetParentSequenceFrame() which returns a time local to the input section in the space of its parent sequence

diff --git a/Engine/Plugins/MovieScene/SequencerScripting/Source/SequencerScripting/Private/ExtensionLibraries/MovieSceneSectionExtensions.cpp b/Engine/Plugins/MovieScene/SequencerScripting/Source/SequencerScripting/Private/ExtensionLibraries/MovieSceneSectionExtensions.cpp
--- a/Engine/Plugins/MovieScene/SequencerScripting/Source/SequencerScripting/Private/ExtensionLibraries/MovieSceneSectionExtensions.cpp
+++ b/Engine/Plugins/MovieScene/SequencerScripting/Source/SequencerScripting/Private/ExtensionLibraries/MovieSceneSectionExtensions.cpp
@@ -15,6 +15,7 @@
 #include "KeysAndChannels/MovieSceneScriptingString.h"
 #include "KeysAndChannels/MovieSceneScriptingEvent.h"
 #include "KeysAndChannels/MovieSceneScriptingActorReference.h"
+#include "Sections/MovieSceneSubSection.h"
 
 
 FSequencerScriptingRange UMovieSceneSectionExtensions::GetRange(UMovieSceneSection* Section)
@@ -273,3 +274,51 @@ TArray<UMovieSceneScriptingChannel*> UMovieSceneSectionExtensions::FindChannelsB
 
 	return Channels;
 }
+
+
+bool GetSubSectionChain(UMovieSceneSubSection* InSubSection, UMovieSceneSequence* ParentSequence, TArray<UMovieSceneSubSection*>& SubSectionChain)
+{
+	UMovieScene* ParentMovieScene = ParentSequence->GetMovieScene();
+	for (UMovieSceneTrack* MasterTrack : ParentMovieScene->GetMasterTracks())
+	{
+		for (UMovieSceneSection* Section : MasterTrack->GetAllSections())
+		{
+			if (Section == InSubSection)
+			{
+				SubSectionChain.Add(InSubSection);
+				return true;
+			}
+			if (Section->IsA(UMovieSceneSubSection::StaticClass()))
+			{
+				UMovieSceneSubSection* SubSection = Cast<UMovieSceneSubSection>(Section);
+				if (GetSubSectionChain(InSubSection, SubSection->GetSequence(), SubSectionChain))
+				{
+					SubSectionChain.Add(SubSection);
+				}
+			}
+		}
+	}
+	return false;
+}
+
+
+int32 UMovieSceneSectionExtensions::GetParentSequenceFrame(UMovieSceneSubSection* InSubSection, int32 InFrame, UMovieSceneSequence* ParentSequence)
+{
+	TArray<UMovieSceneSubSection*> SubSectionChain;
+	GetSubSectionChain(InSubSection, ParentSequence, SubSectionChain);
+		
+	FFrameRate LocalDisplayRate = InSubSection->GetSequence()->GetMovieScene()->GetDisplayRate();
+	FFrameRate LocalTickResolution = InSubSection->GetSequence()->GetMovieScene()->GetTickResolution();
+	FFrameTime LocalFrameTime = ConvertFrameTime(InFrame, LocalDisplayRate, LocalTickResolution);
+		
+	for (int32 SectionIndex = 0; SectionIndex < SubSectionChain.Num(); ++SectionIndex)
+	{
+		LocalFrameTime = LocalFrameTime * SubSectionChain[SectionIndex]->OuterToInnerTransform().Inverse();
+	}
+
+	FFrameRate ParentDisplayRate = ParentSequence->GetMovieScene()->GetDisplayRate();
+	FFrameRate ParentTickResolution = ParentSequence->GetMovieScene()->GetTickResolution();
+
+	LocalFrameTime = ConvertFrameTime(LocalFrameTime, ParentTickResolution, ParentDisplayRate);
+	return LocalFrameTime.GetFrame().Value;
+}
diff --git a/Engine/Plugins/MovieScene/SequencerScripting/Source/SequencerScripting/Public/ExtensionLibraries/MovieSceneSectionExtensions.h b/Engine/Plugins/MovieScene/SequencerScripting/Source/SequencerScripting/Public/ExtensionLibraries/MovieSceneSectionExtensions.h
--- a/Engine/Plugins/MovieScene/SequencerScripting/Source/SequencerScripting/Public/ExtensionLibraries/MovieSceneSectionExtensions.h
+++ b/Engine/Plugins/MovieScene/SequencerScripting/Source/SequencerScripting/Public/ExtensionLibraries/MovieSceneSectionExtensions.h
@@ -156,4 +156,14 @@ public:
 	UFUNCTION(BlueprintCallable, Category = Section, meta = (ScriptMethod))
 	static TArray<UMovieSceneScriptingChannel*> FindChannelsByType(UMovieSceneSection* Section, TSubclassOf<UMovieSceneScriptingChannel> ChannelType);
 
+	/**
+	 * Get the frame in the space of its parent sequence
+	 *
+	 * @param Section        The section that the InFrame is local to
+	 * @param InFrame The desired local frame
+	 * @param ParentSequence The parent sequence to traverse from
+	 * @return The frame at the parent sequence
+	 */
+	UFUNCTION(BlueprintCallable, Category = "Section", meta = (ScriptMethod))
+	static int32 GetParentSequenceFrame(UMovieSceneSubSection* Section, int32 InFrame, UMovieSceneSequence* ParentSequence);
 };
\ No newline at end of file
