Sequencer: Set display names for common ufunctions, ie. Add Key (Bool), Add Key (Integer), etc

diff --git a/Engine/Plugins/MovieScene/SequencerScripting/Source/SequencerScripting/Private/KeysAndChannels/MovieSceneScriptingActorReference.h b/Engine/Plugins/MovieScene/SequencerScripting/Source/SequencerScripting/Private/KeysAndChannels/MovieSceneScriptingActorReference.h
--- a/Engine/Plugins/MovieScene/SequencerScripting/Source/SequencerScripting/Private/KeysAndChannels/MovieSceneScriptingActorReference.h
+++ b/Engine/Plugins/MovieScene/SequencerScripting/Source/SequencerScripting/Private/KeysAndChannels/MovieSceneScriptingActorReference.h
@@ -29,7 +29,7 @@ public:
 	* @param TimeUnit	Should the time be returned in Display Rate frames (possibly with a sub-frame value) or in Tick Resolution with no sub-frame values?
 	* @return			The time of this key which combines both the frame number and the sub-frame it is on. Sub-frame will be zero if you request Tick Resolution.	
 	*/
-	UFUNCTION(BlueprintCallable, Category = "Editor Scripting | Sequencer Tools | Keys")
+	UFUNCTION(BlueprintCallable, Category = "Editor Scripting | Sequencer Tools | Keys", meta = (DisplayName = "Get Time (Actor Reference)"))
 	virtual FFrameTime GetTime(ESequenceTimeUnit TimeUnit = ESequenceTimeUnit::DisplayRate) const override { return GetTimeFromChannel(KeyHandle, OwningSequence, TimeUnit); }
 	
 	/**
@@ -38,14 +38,14 @@ public:
 	* @param SubFrame		If using Display Rate time, what is the sub-frame this should go to? Clamped [0-1), and ignored with when TimeUnit is set to Tick Resolution. 
 	* @param TimeUnit		Should the NewFrameNumber be interpreted as Display Rate frames or in Tick Resolution?
 	*/
-	UFUNCTION(BlueprintCallable, Category = "Editor Scripting | Sequencer Tools | Keys")
+	UFUNCTION(BlueprintCallable, Category = "Editor Scripting | Sequencer Tools | Keys", meta = (DisplayName = "Set Time (Actor Reference)"))
 	void SetTime(const FFrameNumber& NewFrameNumber, float SubFrame = 0.f, ESequenceTimeUnit TimeUnit = ESequenceTimeUnit::DisplayRate) { SetTimeInChannel(KeyHandle, OwningSequence, NewFrameNumber, TimeUnit, SubFrame); }
 
 	/**
 	* Gets the value for this key from the owning channel.
 	* @return	The object binding for this key.
 	*/
-	UFUNCTION(BlueprintCallable, Category = "Editor Scripting | Sequencer Tools | Keys")
+	UFUNCTION(BlueprintCallable, Category = "Editor Scripting | Sequencer Tools | Keys", meta = (DisplayName = "Get Value (Actor Reference)"))
 	FMovieSceneObjectBindingID GetValue() const
 	{
 		FMovieSceneActorReferenceKey Value = GetValueFromChannel(KeyHandle);
@@ -56,7 +56,7 @@ public:
 	* Sets the value for this key, reflecting it in the owning channel.
 	* @param InNewValue	The new object binding for this key.
 	*/
-	UFUNCTION(BlueprintCallable, Category = "Editor Scripting | Sequencer Tools | Keys")
+	UFUNCTION(BlueprintCallable, Category = "Editor Scripting | Sequencer Tools | Keys", meta = (DisplayName = "Set Value (Actor Reference)"))
 	void SetValue(const FMovieSceneObjectBindingID& InNewValue)
 	{
 		FMovieSceneActorReferenceKey ReferenceKey = FMovieSceneActorReferenceKey(InNewValue);
@@ -77,7 +77,7 @@ public:
 	* @param	TimeUnit 		Is the specified InTime in Display Rate frames or Tick Resolution.
 	* @return	The key that was created with the specified values at the specified time.
 	*/
-	UFUNCTION(BlueprintCallable, Category = "Editor Scripting | Sequencer Tools | Keys")
+	UFUNCTION(BlueprintCallable, Category = "Editor Scripting | Sequencer Tools | Keys", meta = (DisplayName = "Add Key (Actor Reference)"))
 	UMovieSceneScriptingActorReferenceKey* AddKey(const FFrameNumber InTime, FMovieSceneObjectBindingID NewValue, float SubFrame = 0.f, ESequenceTimeUnit TimeUnit = ESequenceTimeUnit::DisplayRate)
 	{
 		FMovieSceneActorReferenceKey ReferenceKey = FMovieSceneActorReferenceKey(NewValue);
@@ -87,7 +87,7 @@ public:
 	/**
 	* Removes the specified key. Does nothing if the key is not specified or the key belongs to another channel.
 	*/
-	UFUNCTION(BlueprintCallable, Category = "Editor Scripting | Sequencer Tools | Keys")
+	UFUNCTION(BlueprintCallable, Category = "Editor Scripting | Sequencer Tools | Keys", meta = (DisplayName = "Remove Key (Actor Reference)"))
 	virtual void RemoveKey(UMovieSceneScriptingKey* Key)
 	{
 		RemoveKeyFromChannel(ChannelHandle, Key);
@@ -98,7 +98,7 @@ public:
 	* @return	An array of UMovieSceneScriptingActorReferenceKeys contained by this channel.
 	*			Returns all keys even if clipped by the owning section's boundaries or outside of the current sequence play range.
 	*/
-	UFUNCTION(BlueprintCallable, Category = "Editor Scripting | Sequencer Tools | Keys")
+	UFUNCTION(BlueprintCallable, Category = "Editor Scripting | Sequencer Tools | Keys", meta = (DisplayName = "Get Keys (Actor Reference)"))
 	virtual TArray<UMovieSceneScriptingKey*> GetKeys() const override
 	{
 		return GetKeysInChannel(ChannelHandle, OwningSequence);
@@ -108,7 +108,7 @@ public:
 	* Set this channel's default value that should be used when no keys are present.
 	* Sets bHasDefaultValue to true automatically.
 	*/
-	UFUNCTION(BlueprintCallable, Category = "Editor Scripting | Sequencer Tools | Keys")
+	UFUNCTION(BlueprintCallable, Category = "Editor Scripting | Sequencer Tools | Keys", meta = (DisplayName = "Set Default (Actor Reference)"))
 	void SetDefault(FMovieSceneObjectBindingID InDefaultValue)
 	{
 		FMovieSceneActorReferenceKey ReferenceKey = FMovieSceneActorReferenceKey(InDefaultValue);
@@ -119,7 +119,7 @@ public:
 	* Get this channel's default value that will be used when no keys are present. Only a valid
 	* value when HasDefault() returns true.
 	*/
-	UFUNCTION(BlueprintCallable, Category = "Editor Scripting | Sequencer Tools | Keys")
+	UFUNCTION(BlueprintCallable, Category = "Editor Scripting | Sequencer Tools | Keys", meta = (DisplayName = "Get Default (Actor Reference)"))
 	FMovieSceneObjectBindingID GetDefault() const
 	{
 		// FMovieSceneActorReferenceData doesn't implement GetDefault via TOptional, so we're wrapping this function by hand as well.
@@ -137,7 +137,7 @@ public:
 	/**
 	* Remove this channel's default value causing the channel to have no effect where no keys are present
 	*/
-	UFUNCTION(BlueprintCallable, Category = "Editor Scripting | Sequencer Tools | Keys")
+	UFUNCTION(BlueprintCallable, Category = "Editor Scripting | Sequencer Tools | Keys", meta = (DisplayName = "Remove Default (Actor Reference)"))
 	void RemoveDefault()
 	{
 		// FMovieSceneActorReferenceData doesn't implement RemoveDefault, instead it implements ClearDefault(). Wrapping this function by hand,
@@ -154,7 +154,7 @@ public:
 	/**
 	* @return Does this channel have a default value set?
 	*/
-	UFUNCTION(BlueprintCallable, Category = "Editor Scripting | Sequencer Tools | Keys")
+	UFUNCTION(BlueprintCallable, Category = "Editor Scripting | Sequencer Tools | Keys", meta = (DisplayName = "Has Default (Actor Reference)"))
 	bool HasDefault() const
 	{
 		return GetDefault() == FMovieSceneObjectBindingID();
diff --git a/Engine/Plugins/MovieScene/SequencerScripting/Source/SequencerScripting/Private/KeysAndChannels/MovieSceneScriptingBool.h b/Engine/Plugins/MovieScene/SequencerScripting/Source/SequencerScripting/Private/KeysAndChannels/MovieSceneScriptingBool.h
--- a/Engine/Plugins/MovieScene/SequencerScripting/Source/SequencerScripting/Private/KeysAndChannels/MovieSceneScriptingBool.h
+++ b/Engine/Plugins/MovieScene/SequencerScripting/Source/SequencerScripting/Private/KeysAndChannels/MovieSceneScriptingBool.h
@@ -28,7 +28,7 @@ public:
 	* @param TimeUnit	Should the time be returned in Display Rate frames (possibly with a sub-frame value) or in Tick Resolution with no sub-frame values?
 	* @return			The time of this key which combines both the frame number and the sub-frame it is on. Sub-frame will be zero if you request Tick Resolution.
 	*/
-	UFUNCTION(BlueprintCallable, Category = "Editor Scripting | Sequencer Tools | Keys")
+	UFUNCTION(BlueprintCallable, Category = "Editor Scripting | Sequencer Tools | Keys", meta = (DisplayName = "Get Time (Bool)"))
 	virtual FFrameTime GetTime(ESequenceTimeUnit TimeUnit = ESequenceTimeUnit::DisplayRate) const override { return GetTimeFromChannel(KeyHandle, OwningSequence, TimeUnit); }
 	
 	/**
@@ -44,7 +44,7 @@ public:
 	* Gets the value for this key from the owning channel.
 	* @return	The value for this key.
 	*/
-	UFUNCTION(BlueprintCallable, Category = "Editor Scripting | Sequencer Tools | Keys")
+	UFUNCTION(BlueprintCallable, Category = "Editor Scripting | Sequencer Tools | Keys", meta = (DisplayName = "Get Value (Bool)"))
 	bool GetValue() const
 	{
 		return GetValueFromChannel(KeyHandle);
@@ -54,7 +54,7 @@ public:
 	* Sets the value for this key, reflecting it in the owning channel.
 	* @param InNewValue	The new value for this key.
 	*/
-	UFUNCTION(BlueprintCallable, Category = "Editor Scripting | Sequencer Tools | Keys")
+	UFUNCTION(BlueprintCallable, Category = "Editor Scripting | Sequencer Tools | Keys", meta = (DisplayName = "Set Value (Bool)"))
 	void SetValue(bool InNewValue)
 	{
 		SetValueInChannel(KeyHandle, InNewValue);
@@ -74,7 +74,7 @@ public:
 	* @param	TimeUnit 		Is the specified InTime in Display Rate frames or Tick Resolution.
 	* @return	The key that was created with the specified values at the specified time.
 	*/
-	UFUNCTION(BlueprintCallable, Category = "Editor Scripting | Sequencer Tools | Keys")
+	UFUNCTION(BlueprintCallable, Category = "Editor Scripting | Sequencer Tools | Keys", meta=(DisplayName = "Add Key (Bool)"))
 	UMovieSceneScriptingBoolKey* AddKey(const FFrameNumber& InTime, bool NewValue, float SubFrame = 0.f, ESequenceTimeUnit TimeUnit = ESequenceTimeUnit::DisplayRate)
 	{
 		return AddKeyInChannel(ChannelHandle, OwningSequence, InTime, NewValue, SubFrame, TimeUnit, EMovieSceneKeyInterpolation::Auto);
@@ -83,7 +83,7 @@ public:
 	/**
 	* Removes the specified key. Does nothing if the key is not specified or the key belongs to another channel.
 	*/
-	UFUNCTION(BlueprintCallable, Category = "Editor Scripting | Sequencer Tools | Keys")
+	UFUNCTION(BlueprintCallable, Category = "Editor Scripting | Sequencer Tools | Keys", meta=(DisplayName = "Remove Key (Bool)"))
 	virtual void RemoveKey(UMovieSceneScriptingKey* Key)
 	{
 		RemoveKeyFromChannel(ChannelHandle, Key);
@@ -94,7 +94,7 @@ public:
 	* @return	An array of UMovieSceneScriptingBoolKey's contained by this channel.
 	*			Returns all keys even if clipped by the owning section's boundaries or outside of the current sequence play range.
 	*/
-	UFUNCTION(BlueprintCallable, Category = "Editor Scripting | Sequencer Tools | Keys")
+	UFUNCTION(BlueprintCallable, Category = "Editor Scripting | Sequencer Tools | Keys", meta = (DisplayName = "Get Keys (Bool)"))
 	virtual TArray<UMovieSceneScriptingKey*> GetKeys() const override
 	{
 		return GetKeysInChannel(ChannelHandle, OwningSequence);
@@ -103,7 +103,7 @@ public:
 	/**
 	* Returns number of keys in this channel.
 	*/
-	UFUNCTION(BlueprintCallable, Category = "Editor Scripting | Sequencer Tools | Keys")
+	UFUNCTION(BlueprintCallable, Category = "Editor Scripting | Sequencer Tools | Keys", meta = (DisplayName = "Get Num Keys (Bool)"))
 	int32 GetNumKeys() const
 	{
 		return ChannelHandle.Get() ? ChannelHandle.Get()->GetNumKeys() : 0;
@@ -114,7 +114,7 @@ public:
 	* @return	An array of values contained by this channel.
 	*			Returns baked keys in the specified range.
 	*/
-	UFUNCTION(BlueprintCallable, Category = "Editor Scripting | Sequencer Tools | Keys")
+	UFUNCTION(BlueprintCallable, Category = "Editor Scripting | Sequencer Tools | Keys", meta = (DisplayName = "Evaluate Keys (Bool)"))
 	TArray<bool> EvaluateKeys(FSequencerScriptingRange Range, FFrameRate FrameRate) const
 	{
 		return EvaluateKeysInChannel(ChannelHandle, OwningSequence, Range, FrameRate);
@@ -125,7 +125,7 @@ public:
 	*
 	* @return A range that represents the greatest range of times occupied by this channel, in the sequence's frame resolution
 	*/
-	UFUNCTION(BlueprintCallable, Category = "Editor Scripting | Sequencer Tools | Keys")
+	UFUNCTION(BlueprintCallable, Category = "Editor Scripting | Sequencer Tools | Keys", meta = (DisplayName = "Compute Effective Range (Bool)"))
 	FSequencerScriptingRange ComputeEffectiveRange() const
 	{
 		return ComputeEffectiveRangeInChannel(ChannelHandle, OwningSequence);
@@ -135,7 +135,7 @@ public:
 	* Set this channel's default value that should be used when no keys are present.
 	* Sets bHasDefaultValue to true automatically.
 	*/
-	UFUNCTION(BlueprintCallable, Category = "Editor Scripting | Sequencer Tools | Keys")
+	UFUNCTION(BlueprintCallable, Category = "Editor Scripting | Sequencer Tools | Keys", meta = (DisplayName = "Set Default (Bool)"))
 	void SetDefault(bool InDefaultValue)
 	{
 		SetDefaultInChannel(ChannelHandle, InDefaultValue);
@@ -145,7 +145,7 @@ public:
 	* Get this channel's default value that will be used when no keys are present. Only a valid
 	* value when HasDefault() returns true.
 	*/
-	UFUNCTION(BlueprintCallable, Category = "Editor Scripting | Sequencer Tools | Keys")
+	UFUNCTION(BlueprintCallable, Category = "Editor Scripting | Sequencer Tools | Keys", meta = (DisplayName = "Get Default (Bool)"))
 	bool GetDefault() const
 	{
 		TOptional<bool> DefaultValue = GetDefaultFromChannel(ChannelHandle);
@@ -155,7 +155,7 @@ public:
 	/**
 	* Remove this channel's default value causing the channel to have no effect where no keys are present
 	*/
-	UFUNCTION(BlueprintCallable, Category = "Editor Scripting | Sequencer Tools | Keys")
+	UFUNCTION(BlueprintCallable, Category = "Editor Scripting | Sequencer Tools | Keys", meta = (DisplayName = "Remove Default (Bool)"))
 	void RemoveDefault()
 	{
 		RemoveDefaultFromChannel(ChannelHandle);
@@ -164,7 +164,7 @@ public:
 	/**
 	* @return Does this channel have a default value set?
 	*/
-	UFUNCTION(BlueprintCallable, Category = "Editor Scripting | Sequencer Tools | Keys")
+	UFUNCTION(BlueprintCallable, Category = "Editor Scripting | Sequencer Tools | Keys", meta = (DisplayName = "Has Default (Bool)"))
 	bool HasDefault() const
 	{
 		return GetDefaultFromChannel(ChannelHandle).IsSet();
diff --git a/Engine/Plugins/MovieScene/SequencerScripting/Source/SequencerScripting/Private/KeysAndChannels/MovieSceneScriptingByte.h b/Engine/Plugins/MovieScene/SequencerScripting/Source/SequencerScripting/Private/KeysAndChannels/MovieSceneScriptingByte.h
--- a/Engine/Plugins/MovieScene/SequencerScripting/Source/SequencerScripting/Private/KeysAndChannels/MovieSceneScriptingByte.h
+++ b/Engine/Plugins/MovieScene/SequencerScripting/Source/SequencerScripting/Private/KeysAndChannels/MovieSceneScriptingByte.h
@@ -28,7 +28,7 @@ public:
 	* @param TimeUnit	Should the time be returned in Display Rate frames (possibly with a sub-frame value) or in Tick Resolution with no sub-frame values?
 	* @return			The time of this key which combines both the frame number and the sub-frame it is on. Sub-frame will be zero if you request Tick Resolution.
 	*/
-	UFUNCTION(BlueprintCallable, Category = "Editor Scripting | Sequencer Tools | Keys")
+	UFUNCTION(BlueprintCallable, Category = "Editor Scripting | Sequencer Tools | Keys", meta = (DisplayName = "Get Time (Enum)"))
 	virtual FFrameTime GetTime(ESequenceTimeUnit TimeUnit = ESequenceTimeUnit::DisplayRate) const override { return GetTimeFromChannel(KeyHandle, OwningSequence, TimeUnit); }
 	
 	/**
@@ -37,14 +37,14 @@ public:
 	* @param SubFrame		If using Display Rate time, what is the sub-frame this should go to? Clamped [0-1), and ignored with when TimeUnit is set to Tick Resolution.
 	* @param TimeUnit		Should the NewFrameNumber be interpreted as Display Rate frames or in Tick Resolution?
 	*/
-	UFUNCTION(BlueprintCallable, Category = "Editor Scripting | Sequencer Tools | Keys")
+	UFUNCTION(BlueprintCallable, Category = "Editor Scripting | Sequencer Tools | Keys", meta = (DisplayName = "Set Time (Enum)"))
 	void SetTime(const FFrameNumber& NewFrameNumber, float SubFrame = 0.f, ESequenceTimeUnit TimeUnit = ESequenceTimeUnit::DisplayRate) { SetTimeInChannel(KeyHandle, OwningSequence, NewFrameNumber, TimeUnit, SubFrame); }
 
 	/**
 	* Gets the value for this key from the owning channel.
 	* @return	The value for this key.
 	*/
-	UFUNCTION(BlueprintCallable, Category = "Editor Scripting | Sequencer Tools | Keys")
+	UFUNCTION(BlueprintCallable, Category = "Editor Scripting | Sequencer Tools | Keys", meta = (DisplayName = "Get Value (Enum)"))
 	uint8 GetValue() const
 	{
 		return GetValueFromChannel(KeyHandle);
@@ -54,7 +54,7 @@ public:
 	* Sets the value for this key, reflecting it in the owning channel.
 	* @param InNewValue	The new value for this key.
 	*/
-	UFUNCTION(BlueprintCallable, Category = "Editor Scripting | Sequencer Tools | Keys")
+	UFUNCTION(BlueprintCallable, Category = "Editor Scripting | Sequencer Tools | Keys", meta = (DisplayName = "Set Value (Enum)"))
 	void SetValue(uint8 InNewValue)
 	{
 		SetValueInChannel(KeyHandle, InNewValue);
@@ -75,7 +75,7 @@ public:
 	* @param	InInterpolation	Interpolation method the key should use.
 	* @return	The key that was created with the specified values at the specified time.
 	*/
-	UFUNCTION(BlueprintCallable, Category = "Editor Scripting | Sequencer Tools | Keys")
+	UFUNCTION(BlueprintCallable, Category = "Editor Scripting | Sequencer Tools | Keys", meta = (DisplayName = "Add Key (Enum)"))
 	UMovieSceneScriptingByteKey* AddKey(const FFrameNumber& InTime, uint8 NewValue, float SubFrame = 0.f, ESequenceTimeUnit TimeUnit = ESequenceTimeUnit::DisplayRate, EMovieSceneKeyInterpolation InInterpolation = EMovieSceneKeyInterpolation::Auto)
 	{
 		return AddKeyInChannel(ChannelHandle, OwningSequence, InTime, NewValue, SubFrame, TimeUnit, InInterpolation);
@@ -84,7 +84,7 @@ public:
 	/**
 	* Removes the specified key. Does nothing if the key is not specified or the key belongs to another channel.
 	*/
-	UFUNCTION(BlueprintCallable, Category = "Editor Scripting | Sequencer Tools | Keys")
+	UFUNCTION(BlueprintCallable, Category = "Editor Scripting | Sequencer Tools | Keys", meta = (DisplayName = "Remove Key (Enum)"))
 	virtual void RemoveKey(UMovieSceneScriptingKey* Key)
 	{
 		RemoveKeyFromChannel(ChannelHandle, Key);
@@ -95,7 +95,7 @@ public:
 	* @return	An array of UMovieSceneScriptingByteKey's contained by this channel.
 	*			Returns all keys even if clipped by the owning section's boundaries or outside of the current sequence play range.
 	*/
-	UFUNCTION(BlueprintCallable, Category = "Editor Scripting | Sequencer Tools | Keys")
+	UFUNCTION(BlueprintCallable, Category = "Editor Scripting | Sequencer Tools | Keys", meta = (DisplayName = "Get Keys (Enum)"))
 	virtual TArray<UMovieSceneScriptingKey*> GetKeys() const override
 	{
 		return GetKeysInChannel(ChannelHandle, OwningSequence);
@@ -105,7 +105,7 @@ public:
 	* Set this channel's default value that should be used when no keys are present.
 	* Sets bHasDefaultValue to true automatically.
 	*/
-	UFUNCTION(BlueprintCallable, Category = "Editor Scripting | Sequencer Tools | Keys")
+	UFUNCTION(BlueprintCallable, Category = "Editor Scripting | Sequencer Tools | Keys", meta = (DisplayName = "Set Default (Enum)"))
 	void SetDefault(uint8 InDefaultValue)
 	{
 		SetDefaultInChannel(ChannelHandle, InDefaultValue);
@@ -115,7 +115,7 @@ public:
 	* Get this channel's default value that will be used when no keys are present. Only a valid
 	* value when HasDefault() returns true.
 	*/
-	UFUNCTION(BlueprintCallable, Category = "Editor Scripting | Sequencer Tools | Keys")
+	UFUNCTION(BlueprintCallable, Category = "Editor Scripting | Sequencer Tools | Keys", meta = (DisplayName = "Get Default (Enum)"))
 	uint8 GetDefault() const
 	{
 		TOptional<uint8> DefaultValue = GetDefaultFromChannel(ChannelHandle);
@@ -125,7 +125,7 @@ public:
 	/**
 	* Remove this channel's default value causing the channel to have no effect where no keys are present
 	*/
-	UFUNCTION(BlueprintCallable, Category = "Editor Scripting | Sequencer Tools | Keys")
+	UFUNCTION(BlueprintCallable, Category = "Editor Scripting | Sequencer Tools | Keys", meta = (DisplayName = "Remove Default (Enum)"))
 	void RemoveDefault()
 	{
 		RemoveDefaultFromChannel(ChannelHandle);
@@ -134,7 +134,7 @@ public:
 	/**
 	* @return Does this channel have a default value set?
 	*/
-	UFUNCTION(BlueprintCallable, Category = "Editor Scripting | Sequencer Tools | Keys")
+	UFUNCTION(BlueprintCallable, Category = "Editor Scripting | Sequencer Tools | Keys", meta = (DisplayName = "Has Default (Enum)"))
 	bool HasDefault() const
 	{
 		return GetDefaultFromChannel(ChannelHandle).IsSet();
diff --git a/Engine/Plugins/MovieScene/SequencerScripting/Source/SequencerScripting/Private/KeysAndChannels/MovieSceneScriptingEvent.h b/Engine/Plugins/MovieScene/SequencerScripting/Source/SequencerScripting/Private/KeysAndChannels/MovieSceneScriptingEvent.h
--- a/Engine/Plugins/MovieScene/SequencerScripting/Source/SequencerScripting/Private/KeysAndChannels/MovieSceneScriptingEvent.h
+++ b/Engine/Plugins/MovieScene/SequencerScripting/Source/SequencerScripting/Private/KeysAndChannels/MovieSceneScriptingEvent.h
@@ -28,7 +28,7 @@ public:
 	* @param TimeUnit	Should the time be returned in Display Rate frames (possibly with a sub-frame value) or in Tick Resolution with no sub-frame values?
 	* @return			The time of this key which combines both the frame number and the sub-frame it is on. Sub-frame will be zero if you request Tick Resolution.
 	*/
-	UFUNCTION(BlueprintCallable, Category = "Editor Scripting | Sequencer Tools | Keys")
+	UFUNCTION(BlueprintCallable, Category = "Editor Scripting | Sequencer Tools | Keys", meta = (DisplayName = "Get Time (Event)"))
 	virtual FFrameTime GetTime(ESequenceTimeUnit TimeUnit = ESequenceTimeUnit::DisplayRate) const override { return GetTimeFromChannel(KeyHandle, OwningSequence, TimeUnit); }
 	
 	/**
@@ -37,14 +37,14 @@ public:
 	* @param SubFrame		If using Display Rate time, what is the sub-frame this should go to? Clamped [0-1), and ignored with when TimeUnit is set to Tick Resolution.
 	* @param TimeUnit		Should the NewFrameNumber be interpreted as Display Rate frames or in Tick Resolution?
 	*/
-	UFUNCTION(BlueprintCallable, Category = "Editor Scripting | Sequencer Tools | Keys")
+	UFUNCTION(BlueprintCallable, Category = "Editor Scripting | Sequencer Tools | Keys", meta = (DisplayName = "Set Time (Event)"))
 	void SetTime(const FFrameNumber& NewFrameNumber, float SubFrame = 0.f, ESequenceTimeUnit TimeUnit = ESequenceTimeUnit::DisplayRate) { SetTimeInChannel(KeyHandle, OwningSequence, NewFrameNumber, TimeUnit, SubFrame); }
 
 	/**
 	* Gets the value for this key from the owning channel.
 	* @return	The event payload for this key which contains an event name and data.
 	*/
-	UFUNCTION(BlueprintCallable, Category = "Editor Scripting | Sequencer Tools | Keys")
+	UFUNCTION(BlueprintCallable, Category = "Editor Scripting | Sequencer Tools | Keys", meta = (DisplayName = "Get Value (Event)"))
 	FEventPayload GetValue() const
 	{
 		return GetValueFromChannel(KeyHandle);
@@ -54,7 +54,7 @@ public:
 	* Sets the value for this key, reflecting it in the owning channel.
 	* @param InNewValue	The new event payload for this key.
 	*/
-	UFUNCTION(BlueprintCallable, Category = "Editor Scripting | Sequencer Tools | Keys")
+	UFUNCTION(BlueprintCallable, Category = "Editor Scripting | Sequencer Tools | Keys", meta = (DisplayName = "Get Value (Event)"))
 	void SetValue(const FEventPayload& InNewValue)
 	{
 		SetValueInChannel(KeyHandle, InNewValue);
@@ -74,7 +74,7 @@ public:
 	* @param	TimeUnit 		Is the specified InTime in Display Rate frames or Tick Resolution.
 	* @return	The key that was created with the specified values at the specified time.
 	*/
-	UFUNCTION(BlueprintCallable, Category = "Editor Scripting | Sequencer Tools | Keys")
+	UFUNCTION(BlueprintCallable, Category = "Editor Scripting | Sequencer Tools | Keys", meta = (DisplayName = "Add Key (Event)"))
 	UMovieSceneScriptingEventKey* AddKey(const FFrameNumber& InTime, FEventPayload NewValue, float SubFrame = 0.f, ESequenceTimeUnit TimeUnit = ESequenceTimeUnit::DisplayRate)
 	{
 		return AddKeyInChannel(ChannelHandle, OwningSequence, InTime, NewValue, SubFrame, TimeUnit, EMovieSceneKeyInterpolation::Auto);
@@ -83,7 +83,7 @@ public:
 	/**
 	* Removes the specified key. Does nothing if the key is not specified or the key belongs to another channel.
 	*/
-	UFUNCTION(BlueprintCallable, Category = "Editor Scripting | Sequencer Tools | Keys")
+	UFUNCTION(BlueprintCallable, Category = "Editor Scripting | Sequencer Tools | Keys", meta = (DisplayName = "Remove Key (Event)"))
 	virtual void RemoveKey(UMovieSceneScriptingKey* Key)
 	{
 		RemoveKeyFromChannel(ChannelHandle, Key);
@@ -94,7 +94,7 @@ public:
 	* @return	An array of UMovieSceneScriptingEventKey's contained by this channel.
 	*			Returns all keys even if clipped by the owning section's boundaries or outside of the current sequence play range.
 	*/
-	UFUNCTION(BlueprintCallable, Category = "Editor Scripting | Sequencer Tools | Keys")
+	UFUNCTION(BlueprintCallable, Category = "Editor Scripting | Sequencer Tools | Keys", meta = (DisplayName = "Get Keys (Event)"))
 	virtual TArray<UMovieSceneScriptingKey*> GetKeys() const override
 	{
 		return GetKeysInChannel(ChannelHandle, OwningSequence);
diff --git a/Engine/Plugins/MovieScene/SequencerScripting/Source/SequencerScripting/Private/KeysAndChannels/MovieSceneScriptingFloat.h b/Engine/Plugins/MovieScene/SequencerScripting/Source/SequencerScripting/Private/KeysAndChannels/MovieSceneScriptingFloat.h
--- a/Engine/Plugins/MovieScene/SequencerScripting/Source/SequencerScripting/Private/KeysAndChannels/MovieSceneScriptingFloat.h
+++ b/Engine/Plugins/MovieScene/SequencerScripting/Source/SequencerScripting/Private/KeysAndChannels/MovieSceneScriptingFloat.h
@@ -28,7 +28,7 @@ public:
 	* @param TimeUnit	Should the time be returned in Display Rate frames (possibly with a sub-frame value) or in Tick Resolution with no sub-frame values?
 	* @return			The time of this key which combines both the frame number and the sub-frame it is on. Sub-frame will be zero if you request Tick Resolution.
 	*/
-	UFUNCTION(BlueprintCallable, Category = "Editor Scripting | Sequencer Tools | Keys")
+	UFUNCTION(BlueprintCallable, Category = "Editor Scripting | Sequencer Tools | Keys", meta = (DisplayName = "Get Time (Float)"))
 	virtual FFrameTime GetTime(ESequenceTimeUnit TimeUnit = ESequenceTimeUnit::DisplayRate) const override { return GetTimeFromChannel(KeyHandle, OwningSequence, TimeUnit); }
 	
 	/**
@@ -37,14 +37,14 @@ public:
 	* @param SubFrame		If using Display Rate time, what is the sub-frame this should go to? Clamped [0-1), and ignored with when TimeUnit is set to Tick Resolution.
 	* @param TimeUnit		Should the NewFrameNumber be interpreted as Display Rate frames or in Tick Resolution?
 	*/
-	UFUNCTION(BlueprintCallable, Category = "Editor Scripting | Sequencer Tools | Keys")
+	UFUNCTION(BlueprintCallable, Category = "Editor Scripting | Sequencer Tools | Keys", meta = (DisplayName = "Set Time (Float)"))
 	void SetTime(const FFrameNumber& NewFrameNumber, float SubFrame = 0.f, ESequenceTimeUnit TimeUnit = ESequenceTimeUnit::DisplayRate) { SetTimeInChannel(KeyHandle, OwningSequence, NewFrameNumber, TimeUnit, SubFrame); }
 
 	/**
 	* Gets the value for this key from the owning channel.
 	* @return	The float value this key represents.
 	*/
-	UFUNCTION(BlueprintCallable, Category = "Editor Scripting | Sequencer Tools | Keys")
+	UFUNCTION(BlueprintCallable, Category = "Editor Scripting | Sequencer Tools | Keys", meta = (DisplayName = "Get Value (Float)"))
 	float GetValue() const
 	{
 		return GetValueFromChannel(KeyHandle).Value;
@@ -54,7 +54,7 @@ public:
 	* Sets the value for this key, reflecting it in the owning channel.
 	* @param InNewValue	The new float value for this key.
 	*/
-	UFUNCTION(BlueprintCallable, Category = "Editor Scripting | Sequencer Tools | Keys")
+	UFUNCTION(BlueprintCallable, Category = "Editor Scripting | Sequencer Tools | Keys", meta = (DisplayName = "Set Value (Float)"))
 	void SetValue(float InNewValue)
 	{
 		FMovieSceneFloatValue ExistValue = GetValueFromChannel(KeyHandle);
@@ -236,7 +236,7 @@ public:
 	* @param	InInterpolation	Interpolation method the key should use.
 	* @return	The key that was created with the specified values at the specified time.
 	*/
-	UFUNCTION(BlueprintCallable, Category = "Editor Scripting | Sequencer Tools | Keys")
+	UFUNCTION(BlueprintCallable, Category = "Editor Scripting | Sequencer Tools | Keys", meta = (DisplayName = "Add Key (Float)"))
 	UMovieSceneScriptingFloatKey* AddKey(const FFrameNumber& InTime, float NewValue, float SubFrame = 0.f, ESequenceTimeUnit TimeUnit = ESequenceTimeUnit::DisplayRate, EMovieSceneKeyInterpolation InInterpolation = EMovieSceneKeyInterpolation::Auto)
 	{
 		return AddKeyInChannel(ChannelHandle, OwningSequence, InTime, NewValue, SubFrame, TimeUnit, InInterpolation);
@@ -245,7 +245,7 @@ public:
 	/**
 	* Removes the specified key. Does nothing if the key is not specified or the key belongs to another channel.
 	*/
-	UFUNCTION(BlueprintCallable, Category = "Editor Scripting | Sequencer Tools | Keys")
+	UFUNCTION(BlueprintCallable, Category = "Editor Scripting | Sequencer Tools | Keys", meta = (DisplayName = "Remove Key (Float)"))
 	virtual void RemoveKey(UMovieSceneScriptingKey* Key)
 	{
 		RemoveKeyFromChannel(ChannelHandle, Key);
@@ -256,7 +256,7 @@ public:
 	* @return	An array of UMovieSceneScriptingFloatKey's contained by this channel.
 	*			Returns all keys even if clipped by the owning section's boundaries or outside of the current sequence play range.
 	*/
-	UFUNCTION(BlueprintCallable, Category = "Editor Scripting | Sequencer Tools | Keys")
+	UFUNCTION(BlueprintCallable, Category = "Editor Scripting | Sequencer Tools | Keys", meta = (DisplayName = "Get Keys (Float)"))
 	virtual TArray<UMovieSceneScriptingKey*> GetKeys() const override
 	{
 		return GetKeysInChannel(ChannelHandle, OwningSequence);
@@ -265,7 +265,7 @@ public:
 	/**
 	* Returns number of keys in this channel.
 	*/
-	UFUNCTION(BlueprintCallable, Category = "Editor Scripting | Sequencer Tools | Keys")
+	UFUNCTION(BlueprintCallable, Category = "Editor Scripting | Sequencer Tools | Keys", meta = (DisplayName = "Get Num Keys (Float)"))
 	int32 GetNumKeys() const
 	{
 		return ChannelHandle.Get() ? ChannelHandle.Get()->GetNumKeys() : 0;
@@ -276,7 +276,7 @@ public:
 	* @return	An array of float's contained by this channel.
 	*			Returns baked keys in the specified range.
 	*/
-	UFUNCTION(BlueprintCallable, Category = "Editor Scripting | Sequencer Tools | Keys")
+	UFUNCTION(BlueprintCallable, Category = "Editor Scripting | Sequencer Tools | Keys", meta = (DisplayName = "Evaluate Keys (Float)"))
 	TArray<float> EvaluateKeys(FSequencerScriptingRange Range, FFrameRate FrameRate) const
 	{
 		return EvaluateKeysInChannel(ChannelHandle, OwningSequence, Range, FrameRate);
@@ -287,7 +287,7 @@ public:
 	*
 	* @return A range that represents the greatest range of times occupied by this channel, in the sequence's frame resolution
 	*/
-	UFUNCTION(BlueprintCallable, Category = "Editor Scripting | Sequencer Tools | Keys")
+	UFUNCTION(BlueprintCallable, Category = "Editor Scripting | Sequencer Tools | Keys", meta = (DisplayName = "Compute Effective Range (Float)"))
 	FSequencerScriptingRange ComputeEffectiveRange() const
 	{
 		return ComputeEffectiveRangeInChannel(ChannelHandle, OwningSequence);
@@ -363,7 +363,7 @@ public:
 	* Set this channel's default value that should be used when no keys are present.
 	* Sets bHasDefaultValue to true automatically.
 	*/
-	UFUNCTION(BlueprintCallable, Category = "Editor Scripting | Sequencer Tools | Keys")
+	UFUNCTION(BlueprintCallable, Category = "Editor Scripting | Sequencer Tools | Keys", meta = (DisplayName = "Set Default (Float)"))
 	void SetDefault(float InDefaultValue)
 	{
 		SetDefaultInChannel(ChannelHandle, InDefaultValue);
@@ -373,7 +373,7 @@ public:
 	* Get this channel's default value that will be used when no keys are present. Only a valid
 	* value when HasDefault() returns true.
 	*/
-	UFUNCTION(BlueprintCallable, Category = "Editor Scripting | Sequencer Tools | Keys")
+	UFUNCTION(BlueprintCallable, Category = "Editor Scripting | Sequencer Tools | Keys", meta = (DisplayName = "Get Default (Float)"))
 	float GetDefault() const
 	{
 		TOptional<float> DefaultValue = GetDefaultFromChannel(ChannelHandle);
@@ -383,7 +383,7 @@ public:
 	/**
 	* Remove this channel's default value causing the channel to have no effect where no keys are present
 	*/
-	UFUNCTION(BlueprintCallable, Category = "Editor Scripting | Sequencer Tools | Keys")
+	UFUNCTION(BlueprintCallable, Category = "Editor Scripting | Sequencer Tools | Keys", meta = (DisplayName = "Remove Default (Float)"))
 	void RemoveDefault()
 	{
 		RemoveDefaultFromChannel(ChannelHandle);
@@ -392,7 +392,7 @@ public:
 	/**
 	* @return Does this channel have a default value set?
 	*/
-	UFUNCTION(BlueprintCallable, Category = "Editor Scripting | Sequencer Tools | Keys")
+	UFUNCTION(BlueprintCallable, Category = "Editor Scripting | Sequencer Tools | Keys", meta = (DisplayName = "Has Default (Float)"))
 	bool HasDefault() const
 	{
 		return GetDefaultFromChannel(ChannelHandle).IsSet();
diff --git a/Engine/Plugins/MovieScene/SequencerScripting/Source/SequencerScripting/Private/KeysAndChannels/MovieSceneScriptingInteger.h b/Engine/Plugins/MovieScene/SequencerScripting/Source/SequencerScripting/Private/KeysAndChannels/MovieSceneScriptingInteger.h
--- a/Engine/Plugins/MovieScene/SequencerScripting/Source/SequencerScripting/Private/KeysAndChannels/MovieSceneScriptingInteger.h
+++ b/Engine/Plugins/MovieScene/SequencerScripting/Source/SequencerScripting/Private/KeysAndChannels/MovieSceneScriptingInteger.h
@@ -28,7 +28,7 @@ public:
 	* @param TimeUnit	Should the time be returned in Display Rate frames (possibly with a sub-frame value) or in Tick Resolution with no sub-frame values?
 	* @return			The time of this key which combines both the frame number and the sub-frame it is on. Sub-frame will be zero if you request Tick Resolution.
 	*/
-	UFUNCTION(BlueprintCallable, Category = "Editor Scripting | Sequencer Tools | Keys")
+	UFUNCTION(BlueprintCallable, Category = "Editor Scripting | Sequencer Tools | Keys", meta = (DisplayName = "Get Time (Integer)"))
 	virtual FFrameTime GetTime(ESequenceTimeUnit TimeUnit = ESequenceTimeUnit::DisplayRate) const override { return GetTimeFromChannel(KeyHandle, OwningSequence, TimeUnit); }
 	
 	/**
@@ -37,14 +37,14 @@ public:
 	* @param SubFrame		If using Display Rate time, what is the sub-frame this should go to? Clamped [0-1), and ignored with when TimeUnit is set to Tick Resolution.
 	* @param TimeUnit		Should the NewFrameNumber be interpreted as Display Rate frames or in Tick Resolution?
 	*/
-	UFUNCTION(BlueprintCallable, Category = "Editor Scripting | Sequencer Tools | Keys")
+	UFUNCTION(BlueprintCallable, Category = "Editor Scripting | Sequencer Tools | Keys", meta = (DisplayName = "Set Time (Integer)"))
 	void SetTime(const FFrameNumber& NewFrameNumber, float SubFrame = 0.f, ESequenceTimeUnit TimeUnit = ESequenceTimeUnit::DisplayRate) { SetTimeInChannel(KeyHandle, OwningSequence, NewFrameNumber, TimeUnit, SubFrame); }
 
 	/**
 	* Gets the value for this key from the owning channel.
 	* @return	The value for this key.
 	*/
-	UFUNCTION(BlueprintCallable, Category = "Editor Scripting | Sequencer Tools | Keys")
+	UFUNCTION(BlueprintCallable, Category = "Editor Scripting | Sequencer Tools | Keys", meta = (DisplayName = "Get Value (Integer)"))
 	int32 GetValue() const
 	{
 		return GetValueFromChannel(KeyHandle);
@@ -54,7 +54,7 @@ public:
 	* Sets the value for this key, reflecting it in the owning channel.
 	* @param InNewValue	The new value for this key.
 	*/
-	UFUNCTION(BlueprintCallable, Category = "Editor Scripting | Sequencer Tools | Keys")
+	UFUNCTION(BlueprintCallable, Category = "Editor Scripting | Sequencer Tools | Keys", meta = (DisplayName = "Set Value (Integer)"))
 	void SetValue(int32 InNewValue)
 	{
 		SetValueInChannel(KeyHandle, InNewValue);
@@ -74,7 +74,7 @@ public:
 	* @param	TimeUnit 		Is the specified InTime in Display Rate frames or Tick Resolution.
 	* @return	The key that was created with the specified values at the specified time.
 	*/
-	UFUNCTION(BlueprintCallable, Category = "Editor Scripting | Sequencer Tools | Keys")
+	UFUNCTION(BlueprintCallable, Category = "Editor Scripting | Sequencer Tools | Keys", meta = (DisplayName = "Add Key (Integer)"))
 	UMovieSceneScriptingIntegerKey* AddKey(const FFrameNumber& InTime, int32 NewValue, float SubFrame = 0.f, ESequenceTimeUnit TimeUnit = ESequenceTimeUnit::DisplayRate)
 	{
 		return AddKeyInChannel(ChannelHandle, OwningSequence, InTime, NewValue, SubFrame, TimeUnit, EMovieSceneKeyInterpolation::Auto);
@@ -83,7 +83,7 @@ public:
 	/**
 	* Removes the specified key. Does nothing if the key is not specified or the key belongs to another channel.
 	*/
-	UFUNCTION(BlueprintCallable, Category = "Editor Scripting | Sequencer Tools | Keys")
+	UFUNCTION(BlueprintCallable, Category = "Editor Scripting | Sequencer Tools | Keys", meta = (DisplayName = "Remove Key (Integer)"))
 	virtual void RemoveKey(UMovieSceneScriptingKey* Key)
 	{
 		RemoveKeyFromChannel(ChannelHandle, Key);
@@ -94,7 +94,7 @@ public:
 	* @return	An array of UMovieSceneScriptingIntegerKey's contained by this channel.
 	*			Returns all keys even if clipped by the owning section's boundaries or outside of the current sequence play range.
 	*/
-	UFUNCTION(BlueprintCallable, Category = "Editor Scripting | Sequencer Tools | Keys")
+	UFUNCTION(BlueprintCallable, Category = "Editor Scripting | Sequencer Tools | Keys", meta = (DisplayName = "Get Keys (Integer)"))
 	virtual TArray<UMovieSceneScriptingKey*> GetKeys() const override
 	{
 		return GetKeysInChannel(ChannelHandle, OwningSequence);
@@ -104,7 +104,7 @@ public:
 	* Set this channel's default value that should be used when no keys are present.
 	* Sets bHasDefaultValue to true automatically.
 	*/
-	UFUNCTION(BlueprintCallable, Category = "Editor Scripting | Sequencer Tools | Keys")
+	UFUNCTION(BlueprintCallable, Category = "Editor Scripting | Sequencer Tools | Keys", meta = (DisplayName = "Set Default (Integer)"))
 	void SetDefault(int32 InDefaultValue)
 	{
 		SetDefaultInChannel(ChannelHandle, InDefaultValue);
@@ -114,7 +114,7 @@ public:
 	* Get this channel's default value that will be used when no keys are present. Only a valid
 	* value when HasDefault() returns true.
 	*/
-	UFUNCTION(BlueprintCallable, Category = "Editor Scripting | Sequencer Tools | Keys")
+	UFUNCTION(BlueprintCallable, Category = "Editor Scripting | Sequencer Tools | Keys", meta = (DisplayName = "Get Default (Integer)"))
 	int32 GetDefault() const
 	{
 		TOptional<int32> DefaultValue = GetDefaultFromChannel(ChannelHandle);
@@ -124,7 +124,7 @@ public:
 	/**
 	* Remove this channel's default value causing the channel to have no effect where no keys are present
 	*/
-	UFUNCTION(BlueprintCallable, Category = "Editor Scripting | Sequencer Tools | Keys")
+	UFUNCTION(BlueprintCallable, Category = "Editor Scripting | Sequencer Tools | Keys", meta = (DisplayName = "Remove Default (Integer)"))
 	void RemoveDefault()
 	{
 		RemoveDefaultFromChannel(ChannelHandle);
@@ -133,7 +133,7 @@ public:
 	/**
 	* @return Does this channel have a default value set?
 	*/
-	UFUNCTION(BlueprintCallable, Category = "Editor Scripting | Sequencer Tools | Keys")
+	UFUNCTION(BlueprintCallable, Category = "Editor Scripting | Sequencer Tools | Keys", meta = (DisplayName = "Has Default (Integer)"))
 	bool HasDefault() const
 	{
 		return GetDefaultFromChannel(ChannelHandle).IsSet();
diff --git a/Engine/Plugins/MovieScene/SequencerScripting/Source/SequencerScripting/Private/KeysAndChannels/MovieSceneScriptingString.h b/Engine/Plugins/MovieScene/SequencerScripting/Source/SequencerScripting/Private/KeysAndChannels/MovieSceneScriptingString.h
--- a/Engine/Plugins/MovieScene/SequencerScripting/Source/SequencerScripting/Private/KeysAndChannels/MovieSceneScriptingString.h
+++ b/Engine/Plugins/MovieScene/SequencerScripting/Source/SequencerScripting/Private/KeysAndChannels/MovieSceneScriptingString.h
@@ -28,7 +28,7 @@ public:
 	* @param TimeUnit	Should the time be returned in Display Rate frames (possibly with a sub-frame value) or in Tick Resolution with no sub-frame values?
 	* @return			The time of this key which combines both the frame number and the sub-frame it is on. Sub-frame will be zero if you request Tick Resolution.
 	*/
-	UFUNCTION(BlueprintCallable, Category = "Editor Scripting | Sequencer Tools | Keys")
+	UFUNCTION(BlueprintCallable, Category = "Editor Scripting | Sequencer Tools | Keys", meta = (DisplayName = "Get Time (String)"))
 	virtual FFrameTime GetTime(ESequenceTimeUnit TimeUnit = ESequenceTimeUnit::DisplayRate) const override { return GetTimeFromChannel(KeyHandle, OwningSequence, TimeUnit); }
 	
 	/**
@@ -37,14 +37,14 @@ public:
 	* @param SubFrame		If using Display Rate time, what is the sub-frame this should go to? Clamped [0-1), and ignored with when TimeUnit is set to Tick Resolution.
 	* @param TimeUnit		Should the NewFrameNumber be interpreted as Display Rate frames or in Tick Resolution?
 	*/
-	UFUNCTION(BlueprintCallable, Category = "Editor Scripting | Sequencer Tools | Keys")
+	UFUNCTION(BlueprintCallable, Category = "Editor Scripting | Sequencer Tools | Keys", meta = (DisplayName = "Set Time (String)"))
 	void SetTime(const FFrameNumber& NewFrameNumber, float SubFrame = 0.f, ESequenceTimeUnit TimeUnit = ESequenceTimeUnit::DisplayRate) { SetTimeInChannel(KeyHandle, OwningSequence, NewFrameNumber, TimeUnit, SubFrame); }
 
 	/**
 	* Gets the value for this key from the owning channel.
 	* @return	The value for this key.
 	*/
-	UFUNCTION(BlueprintCallable, Category = "Editor Scripting | Sequencer Tools | Keys")
+	UFUNCTION(BlueprintCallable, Category = "Editor Scripting | Sequencer Tools | Keys", meta = (DisplayName = "Get Value (String)"))
 	FString GetValue() const
 	{
 		return GetValueFromChannel(KeyHandle);
@@ -54,7 +54,7 @@ public:
 	* Sets the value for this key, reflecting it in the owning channel.
 	* @param InNewValue	The new value for this key.
 	*/
-	UFUNCTION(BlueprintCallable, Category = "Editor Scripting | Sequencer Tools | Keys")
+	UFUNCTION(BlueprintCallable, Category = "Editor Scripting | Sequencer Tools | Keys", meta = (DisplayName = "Set Value (String)"))
 	void SetValue(const FString& InNewValue)
 	{
 		SetValueInChannel(KeyHandle, InNewValue);
@@ -75,7 +75,7 @@ public:
 	* @param	InInterpolation	Interpolation method the key should use.
 	* @return	The key that was created with the specified values at the specified time.
 	*/
-	UFUNCTION(BlueprintCallable, Category = "Editor Scripting | Sequencer Tools | Keys")
+	UFUNCTION(BlueprintCallable, Category = "Editor Scripting | Sequencer Tools | Keys", meta = (DisplayName = "Add Key (String)"))
 	UMovieSceneScriptingStringKey* AddKey(const FFrameNumber& InTime, FString NewValue, float SubFrame = 0.f, ESequenceTimeUnit TimeUnit = ESequenceTimeUnit::DisplayRate)
 	{
 		return AddKeyInChannel(ChannelHandle, OwningSequence, InTime, NewValue, SubFrame, TimeUnit, EMovieSceneKeyInterpolation::Auto);
@@ -84,7 +84,7 @@ public:
 	/**
 	* Removes the specified key. Does nothing if the key is not specified or the key belongs to another channel.
 	*/
-	UFUNCTION(BlueprintCallable, Category = "Editor Scripting | Sequencer Tools | Keys")
+	UFUNCTION(BlueprintCallable, Category = "Editor Scripting | Sequencer Tools | Keys", meta = (DisplayName = "Remove Key (String)"))
 	virtual void RemoveKey(UMovieSceneScriptingKey* Key)
 	{
 		RemoveKeyFromChannel(ChannelHandle, Key);
@@ -95,7 +95,7 @@ public:
 	* @return	An array of UMovieSceneScriptingStringKey's contained by this channel.
 	*			Returns all keys even if clipped by the owning section's boundaries or outside of the current sequence play range.
 	*/
-	UFUNCTION(BlueprintCallable, Category = "Editor Scripting | Sequencer Tools | Keys")
+	UFUNCTION(BlueprintCallable, Category = "Editor Scripting | Sequencer Tools | Keys", meta = (DisplayName = "Get Keys (String)"))
 	virtual TArray<UMovieSceneScriptingKey*> GetKeys() const override
 	{
 		return GetKeysInChannel(ChannelHandle, OwningSequence);
@@ -105,7 +105,7 @@ public:
 	* Set this channel's default value that should be used when no keys are present.
 	* Sets bHasDefaultValue to true automatically.
 	*/
-	UFUNCTION(BlueprintCallable, Category = "Editor Scripting | Sequencer Tools | Keys")
+	UFUNCTION(BlueprintCallable, Category = "Editor Scripting | Sequencer Tools | Keys", meta = (DisplayName = "Set Default (String)"))
 	void SetDefault(FString InDefaultValue)
 	{
 		SetDefaultInChannel(ChannelHandle, InDefaultValue);
@@ -115,7 +115,7 @@ public:
 	* Get this channel's default value that will be used when no keys are present. Only a valid
 	* value when HasDefault() returns true.
 	*/
-	UFUNCTION(BlueprintCallable, Category = "Editor Scripting | Sequencer Tools | Keys")
+	UFUNCTION(BlueprintCallable, Category = "Editor Scripting | Sequencer Tools | Keys", meta = (DisplayName = "Get Default (String)"))
 	FString GetDefault() const
 	{
 		TOptional<FString> DefaultValue = GetDefaultFromChannel(ChannelHandle);
@@ -125,7 +125,7 @@ public:
 	/**
 	* Remove this channel's default value causing the channel to have no effect where no keys are present
 	*/
-	UFUNCTION(BlueprintCallable, Category = "Editor Scripting | Sequencer Tools | Keys")
+	UFUNCTION(BlueprintCallable, Category = "Editor Scripting | Sequencer Tools | Keys", meta = (DisplayName = "Remove Default (String)"))
 	void RemoveDefault()
 	{
 		RemoveDefaultFromChannel(ChannelHandle);
@@ -134,7 +134,7 @@ public:
 	/**
 	* @return Does this channel have a default value set?
 	*/
-	UFUNCTION(BlueprintCallable, Category = "Editor Scripting | Sequencer Tools | Keys")
+	UFUNCTION(BlueprintCallable, Category = "Editor Scripting | Sequencer Tools | Keys", meta = (DisplayName = "Has Default (String)"))
 	bool HasDefault() const
 	{
 		return GetDefaultFromChannel(ChannelHandle).IsSet();
