Sequencer: Set display names for common ufunctions, ie. Add Key (Bool), Add Key (Integer), etc

diff --git a/Engine/Plugins/MovieScene/SequencerScripting/Source/SequencerScripting/Private/KeysAndChannels/MovieSceneScriptingBool.h b/Engine/Plugins/MovieScene/SequencerScripting/Source/SequencerScripting/Private/KeysAndChannels/MovieSceneScriptingBool.h
--- a/Engine/Plugins/MovieScene/SequencerScripting/Source/SequencerScripting/Private/KeysAndChannels/MovieSceneScriptingBool.h
+++ b/Engine/Plugins/MovieScene/SequencerScripting/Source/SequencerScripting/Private/KeysAndChannels/MovieSceneScriptingBool.h
@@ -37,7 +37,7 @@ public:
 	* @param SubFrame		If using Display Rate time, what is the sub-frame this should go to? Clamped [0-1), and ignored with when TimeUnit is set to Tick Resolution.
 	* @param TimeUnit		Should the NewFrameNumber be interpreted as Display Rate frames or in Tick Resolution?
 	*/
-	UFUNCTION(BlueprintCallable, Category = "Editor Scripting | Sequencer Tools | Keys")
+	UFUNCTION(BlueprintCallable, Category = "Editor Scripting | Sequencer Tools | Keys", meta = (DisplayName = "Set Time (Bool)"))
 	void SetTime(const FFrameNumber& NewFrameNumber, float SubFrame = 0.f, ESequenceTimeUnit TimeUnit = ESequenceTimeUnit::DisplayRate) { SetTimeInChannel(KeyHandle, OwningSequence, NewFrameNumber, TimeUnit, SubFrame); }
 
 	/**
