Sequencer: Fix setting the end playback range to that the end is exclusive.

diff --git a/Engine/Plugins/MovieScene/SequencerScripting/Source/SequencerScripting/Private/ExtensionLibraries/MovieSceneSectionExtensions.cpp b/Engine/Plugins/MovieScene/SequencerScripting/Source/SequencerScripting/Private/ExtensionLibraries/MovieSceneSectionExtensions.cpp
--- a/Engine/Plugins/MovieScene/SequencerScripting/Source/SequencerScripting/Private/ExtensionLibraries/MovieSceneSectionExtensions.cpp
+++ b/Engine/Plugins/MovieScene/SequencerScripting/Source/SequencerScripting/Private/ExtensionLibraries/MovieSceneSectionExtensions.cpp
@@ -85,7 +85,7 @@ void UMovieSceneSectionExtensions::SetRange(UMovieSceneSection* Section, int32 S
 
 		TRange<FFrameNumber> NewRange;
 		NewRange.SetLowerBound(TRangeBound<FFrameNumber>::Inclusive(ConvertFrameTime(StartFrame, DisplayRate, MovieScene->GetTickResolution()).FrameNumber));
-		NewRange.SetUpperBound(TRangeBound<FFrameNumber>::Inclusive(ConvertFrameTime(EndFrame, DisplayRate, MovieScene->GetTickResolution()).FrameNumber));
+		NewRange.SetUpperBound(TRangeBound<FFrameNumber>::Exclusive(ConvertFrameTime(EndFrame, DisplayRate, MovieScene->GetTickResolution()).FrameNumber));
 
 		if (NewRange.GetLowerBound().IsOpen() || NewRange.GetUpperBound().IsOpen() || NewRange.GetLowerBoundValue() <= NewRange.GetUpperBoundValue())
 		{
diff --git a/Engine/Source/Editor/Sequencer/Private/Sequencer.h b/Engine/Source/Editor/Sequencer/Private/Sequencer.h
--- a/Engine/Source/Editor/Sequencer/Private/Sequencer.h
+++ b/Engine/Source/Editor/Sequencer/Private/Sequencer.h
@@ -169,7 +169,7 @@ public:
 	void SetPlaybackRangeEnd()
 	{
 		TRange<FFrameNumber> PlayRange = GetPlaybackRange();
-		SetPlaybackRange(TRange<FFrameNumber>(PlayRange.GetLowerBound(), TRangeBound<FFrameNumber>::Inclusive(GetLocalTime().Time.FrameNumber)));
+		SetPlaybackRange(TRange<FFrameNumber>(PlayRange.GetLowerBound(), TRangeBound<FFrameNumber>::Exclusive(GetLocalTime().Time.FrameNumber)));
 	}
 
 	/**
