Sequencer: Fix call to SetDefaultInChannel, should be SetChannelDefault.

diff --git a/Engine/Plugins/MovieScene/SequencerScripting/Source/SequencerScripting/Private/KeysAndChannels/MovieSceneScriptingChannel.h b/Engine/Plugins/MovieScene/SequencerScripting/Source/SequencerScripting/Private/KeysAndChannels/MovieSceneScriptingChannel.h
--- a/Engine/Plugins/MovieScene/SequencerScripting/Source/SequencerScripting/Private/KeysAndChannels/MovieSceneScriptingChannel.h
+++ b/Engine/Plugins/MovieScene/SequencerScripting/Source/SequencerScripting/Private/KeysAndChannels/MovieSceneScriptingChannel.h
@@ -176,7 +176,7 @@ struct TMovieSceneScriptingChannel
 		if (Channel)
 		{
 			using namespace MovieScene;
-			SetDefaultInChannel(ChannelHandle, InDefaultValue);
+			SetChannelDefault(Channel, InDefaultValue);
 			return;
 		}
 		UE_LOG(LogMovieScene, Error, TEXT("Invalid ChannelHandle for MovieSceneScriptingChannel, failed to set default value."));
