Sequencer: Change log errors to KismetExecutionMessages

diff --git a/Engine/Plugins/MovieScene/SequencerScripting/Source/SequencerScripting/Private/ExtensionLibraries/MovieSceneSectionExtensions.cpp b/Engine/Plugins/MovieScene/SequencerScripting/Source/SequencerScripting/Private/ExtensionLibraries/MovieSceneSectionExtensions.cpp
--- a/Engine/Plugins/MovieScene/SequencerScripting/Source/SequencerScripting/Private/ExtensionLibraries/MovieSceneSectionExtensions.cpp
+++ b/Engine/Plugins/MovieScene/SequencerScripting/Source/SequencerScripting/Private/ExtensionLibraries/MovieSceneSectionExtensions.cpp
@@ -92,7 +92,7 @@ void UMovieSceneSectionExtensions::SetRange(UMovieSceneSection* Section, int32 S
 		}
 		else
 		{
-			UE_LOG(LogMovieScene, Error, TEXT("Invalid range specified"));
+			FFrame::KismetExecutionMessage(TEXT("Invalid range specified"), ELogVerbosity::Error);
 		}
 	}
 }
@@ -114,7 +114,7 @@ void UMovieSceneSectionExtensions::SetRangeSeconds(UMovieSceneSection* Section,
 		}
 		else
 		{
-			UE_LOG(LogMovieScene, Error, TEXT("Invalid range specified"));
+			FFrame::KismetExecutionMessage(TEXT("Invalid range specified"), ELogVerbosity::Error);
 		}
 	}
 }
@@ -224,7 +224,7 @@ TArray<UMovieSceneScriptingChannel*> UMovieSceneSectionExtensions::GetChannels(U
 	TArray<UMovieSceneScriptingChannel*> Channels;
 	if (!Section)
 	{
-		UE_LOG(LogMovieScene, Error, TEXT("Cannot get channels for null section"));
+		FFrame::KismetExecutionMessage(TEXT("Cannot get channels for null section"), ELogVerbosity::Error);
 		return Channels;
 	}
 	FMovieSceneChannelProxy& ChannelProxy = Section->GetChannelProxy();
@@ -249,8 +249,8 @@ TArray<UMovieSceneScriptingChannel*> UMovieSceneSectionExtensions::FindChannelsB
 	TArray<UMovieSceneScriptingChannel*> Channels;
 	if (!Section)
 	{
- 		UE_LOG(LogMovieScene, Error, TEXT("Cannot get channels for null section"));
- 		return Channels;
+		FFrame::KismetExecutionMessage(TEXT("Cannot get channels for null section"), ELogVerbosity::Error);
+		return Channels;
 	}
  
 	FMovieSceneChannelProxy& ChannelProxy = Section->GetChannelProxy();
@@ -267,7 +267,7 @@ TArray<UMovieSceneScriptingChannel*> UMovieSceneSectionExtensions::FindChannelsB
 	else if (ChannelType == UMovieSceneScriptingActorReferenceChannel::StaticClass())	{ GetScriptingChannelsForChannel<FMovieSceneActorReferenceData, UMovieSceneScriptingActorReferenceChannel>(ChannelProxy, Sequence, Channels); }
 	else
 	{
-		UE_LOG(LogMovieScene, Error, TEXT("Unsupported ChannelType for FindChannelsByType!"));
+		FFrame::KismetExecutionMessage(TEXT("Unsupported ChannelType for FindChannelsByType!"), ELogVerbosity::Error);
 	}
 
 
diff --git a/Engine/Plugins/MovieScene/SequencerScripting/Source/SequencerScripting/Private/KeysAndChannels/MovieSceneScriptingActorReference.h b/Engine/Plugins/MovieScene/SequencerScripting/Source/SequencerScripting/Private/KeysAndChannels/MovieSceneScriptingActorReference.h
--- a/Engine/Plugins/MovieScene/SequencerScripting/Source/SequencerScripting/Private/KeysAndChannels/MovieSceneScriptingActorReference.h
+++ b/Engine/Plugins/MovieScene/SequencerScripting/Source/SequencerScripting/Private/KeysAndChannels/MovieSceneScriptingActorReference.h
@@ -130,7 +130,7 @@ public:
 			return ReferenceKey.Object;
 		}
 
-		UE_LOG(LogMovieScene, Error, TEXT("Invalid ChannelHandle for MovieSceneScriptingChannel, failed to get default value."));
+		FFrame::KismetExecutionMessage(TEXT("Invalid ChannelHandle for MovieSceneScriptingChannel, failed to get default value."), ELogVerbosity::Error);
 		return FMovieSceneObjectBindingID();
 	}
 
@@ -148,7 +148,7 @@ public:
 			Channel->ClearDefault();
 			return;
 		}
-		UE_LOG(LogMovieScene, Error, TEXT("Invalid ChannelHandle for MovieSceneScriptingChannel, failed to remove default value."));
+		FFrame::KismetExecutionMessage(TEXT("Invalid ChannelHandle for MovieSceneScriptingChannel, failed to remove default value."), ELogVerbosity::Error);
 	}
 
 	/**
diff --git a/Engine/Plugins/MovieScene/SequencerScripting/Source/SequencerScripting/Private/KeysAndChannels/MovieSceneScriptingChannel.h b/Engine/Plugins/MovieScene/SequencerScripting/Source/SequencerScripting/Private/KeysAndChannels/MovieSceneScriptingChannel.h
--- a/Engine/Plugins/MovieScene/SequencerScripting/Source/SequencerScripting/Private/KeysAndChannels/MovieSceneScriptingChannel.h
+++ b/Engine/Plugins/MovieScene/SequencerScripting/Source/SequencerScripting/Private/KeysAndChannels/MovieSceneScriptingChannel.h
@@ -73,7 +73,7 @@ struct TMovieSceneScriptingChannel
 			return Key;
 		}
 
-		UE_LOG(LogMovieScene, Error, TEXT("Invalid ChannelHandle for MovieSceneScriptingChannel, failed to add key."));
+		FFrame::KismetExecutionMessage(TEXT("Invalid ChannelHandle for MovieSceneScriptingChannel, failed to add key."), ELogVerbosity::Error);
 		return nullptr;
 	}
 
@@ -91,7 +91,7 @@ struct TMovieSceneScriptingChannel
 			return;
 		}
 
-		UE_LOG(LogMovieScene, Error, TEXT("Invalid ChannelHandle for MovieSceneScriptingChannel, failed to remove key."));
+		FFrame::KismetExecutionMessage(TEXT("Invalid ChannelHandle for MovieSceneScriptingChannel, failed to remove key."), ELogVerbosity::Error);
 	}
 
 	TArray<UMovieSceneScriptingKey*> GetKeysInChannel(TMovieSceneChannelHandle<ChannelType> ChannelHandle, TWeakObjectPtr<UMovieSceneSequence> Sequence) const
@@ -115,7 +115,7 @@ struct TMovieSceneScriptingChannel
 		}
 		else
 		{
-			UE_LOG(LogMovieScene, Error, TEXT("Invalid ChannelHandle for MovieSceneScriptingChannel, failed to get keys."));
+			FFrame::KismetExecutionMessage(TEXT("Invalid ChannelHandle for MovieSceneScriptingChannel, failed to get keys."), ELogVerbosity::Error);
 		}
 
 		return OutScriptingKeys;
@@ -143,12 +143,12 @@ struct TMovieSceneScriptingChannel
 			}
 			else
 			{
-				UE_LOG(LogMovieScene, Error, TEXT("Unbounded range passed to evaluate keys."));
+				FFrame::KismetExecutionMessage(TEXT("Unbounded range passed to evaluate keys."), ELogVerbosity::Error);
 			}
 		}
 		else
 		{
-			UE_LOG(LogMovieScene, Error, TEXT("Invalid ChannelHandle for MovieSceneScriptingChannel, failed to evaluate keys."));
+			FFrame::KismetExecutionMessage(TEXT("Invalid ChannelHandle for MovieSceneScriptingChannel, failed to evaluate keys."), ELogVerbosity::Error);
 		}
 
 		return OutValues;
@@ -164,7 +164,7 @@ struct TMovieSceneScriptingChannel
 		}
 		else
 		{
-			UE_LOG(LogMovieScene, Error, TEXT("Invalid ChannelHandle for MovieSceneScriptingChannel, failed to get effective range."));
+			FFrame::KismetExecutionMessage(TEXT("Invalid ChannelHandle for MovieSceneScriptingChannel, failed to get effective range."), ELogVerbosity::Error);
 		}
 
 		return ScriptingRange;
@@ -179,7 +179,7 @@ struct TMovieSceneScriptingChannel
 			SetDefaultInChannel(ChannelHandle, InDefaultValue);
 			return;
 		}
-		UE_LOG(LogMovieScene, Error, TEXT("Invalid ChannelHandle for MovieSceneScriptingChannel, failed to set default value."));
+		FFrame::KismetExecutionMessage(TEXT("Invalid ChannelHandle for MovieSceneScriptingChannel, failed to set default value."), ELogVerbosity::Error);
 	}
 
 	void RemoveDefaultFromChannel(TMovieSceneChannelHandle<ChannelType> ChannelHandle)
@@ -191,7 +191,7 @@ struct TMovieSceneScriptingChannel
 			RemoveChannelDefault(Channel);
 			return;
 		}
-		UE_LOG(LogMovieScene, Error, TEXT("Invalid ChannelHandle for MovieSceneScriptingChannel, failed to remove default value."));
+		FFrame::KismetExecutionMessage(TEXT("Invalid ChannelHandle for MovieSceneScriptingChannel, failed to remove default value."), ELogVerbosity::Error);
 	}
 
 	TOptional<ScriptingKeyValueType> GetDefaultFromChannel(TMovieSceneChannelHandle<ChannelType> ChannelHandle) const
@@ -210,7 +210,7 @@ struct TMovieSceneScriptingChannel
 			return TOptional<ScriptingKeyValueType>();
 		}
 
-		UE_LOG(LogMovieScene, Error, TEXT("Invalid ChannelHandle for MovieSceneScriptingChannel, failed to get default value."));
+		FFrame::KismetExecutionMessage(TEXT("Invalid ChannelHandle for MovieSceneScriptingChannel, failed to get default value."), ELogVerbosity::Error);
 		return TOptional<ScriptingKeyValueType>();
 	}
 };
@@ -226,7 +226,7 @@ struct TMovieSceneScriptingKey
 	{
 		if (!Sequence.IsValid())
 		{
-			UE_LOG(LogMovieScene, Error, TEXT("GetTime called with an invalid owning sequence."));
+			FFrame::KismetExecutionMessage(TEXT("GetTime called with an invalid owning sequence."), ELogVerbosity::Error);
 			return FFrameNumber(0);
 		}
 
@@ -247,7 +247,7 @@ struct TMovieSceneScriptingKey
 			return FFrameTime(KeyTime, 0.f);
 		}
 
-		UE_LOG(LogMovieScene, Error, TEXT("Invalid ChannelHandle for MovieSceneScriptingKey, failed to retrieve Time."));
+		FFrame::KismetExecutionMessage(TEXT("Invalid ChannelHandle for MovieSceneScriptingKey, failed to retrieve Time."), ELogVerbosity::Error);
 		return FFrameTime();
 	}
 
@@ -255,7 +255,7 @@ struct TMovieSceneScriptingKey
 	{
 		if (!Sequence.IsValid())
 		{
-			UE_LOG(LogMovieScene, Error, TEXT("SetTime called with an invalid owning sequence."));
+			FFrame::KismetExecutionMessage(TEXT("SetTime called with an invalid owning sequence."), ELogVerbosity::Error);
 			return;
 		}
 
@@ -265,7 +265,7 @@ struct TMovieSceneScriptingKey
 		// TickResolution doesn't support a sub-frame as you can't get finer detailed than that.
 		if (TimeUnit == ESequenceTimeUnit::TickResolution && SubFrame > 0.f)
 		{
-			UE_LOG(LogMovieScene, Warning, TEXT("SetTime called with a SubFrame specified for a Tick Resolution type time! SubFrames are only allowed for Display Rate types, ignoring..."));
+			FFrame::KismetExecutionMessage(TEXT("SetTime called with a SubFrame specified for a Tick Resolution type time! SubFrames are only allowed for Display Rate types, ignoring..."), ELogVerbosity::Error);
 			SubFrame = 0.f;
 		}
 
@@ -284,7 +284,7 @@ struct TMovieSceneScriptingKey
 			return;
 		}
 
-		UE_LOG(LogMovieScene, Error, TEXT("Invalid ChannelHandle for MovieSceneScriptingKey, failed to set Time."));
+		FFrame::KismetExecutionMessage(TEXT("Invalid ChannelHandle for MovieSceneScriptingKey, failed to set Time."), ELogVerbosity::Error);
 	}
 	
 	ChannelDataType GetValueFromChannel(FKeyHandle KeyHandle) const
@@ -297,14 +297,14 @@ struct TMovieSceneScriptingKey
 			using namespace MovieScene;
 			if (!GetKeyValue(Channel, KeyHandle, Value))
 			{
-				UE_LOG(LogMovieScene, Error, TEXT("Invalid KeyIndex for MovieSceneScriptingKey, failed to get value. Did you forget to create the key through the channel?"));
+				FFrame::KismetExecutionMessage(TEXT("Invalid KeyIndex for MovieSceneScriptingKey, failed to get value. Did you forget to create the key through the channel?"), ELogVerbosity::Error);
 				return Value;
 			}
 
 			return Value;
 		}
 
-		UE_LOG(LogMovieScene, Error, TEXT("Invalid ChannelHandle for MovieSceneScriptingKey, failed to get value. Did you forget to create the key through the channel?"));
+		FFrame::KismetExecutionMessage(TEXT("Invalid ChannelHandle for MovieSceneScriptingKey, failed to get value. Did you forget to create the key through the channel?"), ELogVerbosity::Error);
 		return Value;
 	}
 
@@ -316,12 +316,12 @@ struct TMovieSceneScriptingKey
 			using namespace MovieScene;
 			if(!AssignValue(Channel, KeyHandle, InNewValue))
 			{
-				UE_LOG(LogMovieScene, Error, TEXT("Invalid KeyIndex for MovieSceneScriptingKey, failed to set value. Did you forget to create the key through the channel?"));
+				FFrame::KismetExecutionMessage(TEXT("Invalid KeyIndex for MovieSceneScriptingKey, failed to set value. Did you forget to create the key through the channel?"), ELogVerbosity::Error);
 				return;
 			}
 		}
 
-		UE_LOG(LogMovieScene, Error, TEXT("Invalid ChannelHandle for MovieSceneScriptingKey, failed to set value. Did you forget to create the key through the channel?"));
+		FFrame::KismetExecutionMessage(TEXT("Invalid ChannelHandle for MovieSceneScriptingKey, failed to set value. Did you forget to create the key through the channel?"), ELogVerbosity::Error);
 	}
 
 public:
diff --git a/Engine/Plugins/MovieScene/SequencerScripting/Source/SequencerScripting/Private/KeysAndChannels/MovieSceneScriptingFloat.h b/Engine/Plugins/MovieScene/SequencerScripting/Source/SequencerScripting/Private/KeysAndChannels/MovieSceneScriptingFloat.h
--- a/Engine/Plugins/MovieScene/SequencerScripting/Source/SequencerScripting/Private/KeysAndChannels/MovieSceneScriptingFloat.h
+++ b/Engine/Plugins/MovieScene/SequencerScripting/Source/SequencerScripting/Private/KeysAndChannels/MovieSceneScriptingFloat.h
@@ -305,7 +305,7 @@ public:
 			return Channel->PreInfinityExtrap;
 		}
 
-		UE_LOG(LogMovieScene, Error, TEXT("Invalid ChannelHandle for MovieSceneScriptingChannel, failed to get pre-infinity extrapolation."));
+		FFrame::KismetExecutionMessage(TEXT("Invalid ChannelHandle for MovieSceneScriptingChannel, failed to get pre-infinity extrapolation."), ELogVerbosity::Error);
 		return ERichCurveExtrapolation::RCCE_None;
 	}
 
@@ -323,7 +323,7 @@ public:
 			return;
 		}
 
-		UE_LOG(LogMovieScene, Error, TEXT("Invalid ChannelHandle for MovieSceneScriptingChannel, failed to set pre-infinity extrapolation."));
+		FFrame::KismetExecutionMessage(TEXT("Invalid ChannelHandle for MovieSceneScriptingChannel, failed to set pre-infinity extrapolation."), ELogVerbosity::Error);
 	}
 
 	/**
@@ -338,7 +338,7 @@ public:
 			return Channel->PostInfinityExtrap;
 		}
 
-		UE_LOG(LogMovieScene, Error, TEXT("Invalid ChannelHandle for MovieSceneScriptingChannel, failed to get post-infinity extrapolation."));
+		FFrame::KismetExecutionMessage(TEXT("Invalid ChannelHandle for MovieSceneScriptingChannel, failed to get post-infinity extrapolation."), ELogVerbosity::Error);
 		return ERichCurveExtrapolation::RCCE_None;
 	}
 
@@ -356,7 +356,7 @@ public:
 			return;
 		}
 
-		UE_LOG(LogMovieScene, Error, TEXT("Invalid ChannelHandle for MovieSceneScriptingChannel, failed to set post-infinity extrapolation."));
+		FFrame::KismetExecutionMessage(TEXT("Invalid ChannelHandle for MovieSceneScriptingChannel, failed to set post-infinity extrapolation."), ELogVerbosity::Error);
 	}
 
 	/**
