Sequencer: Safe-guarding against nullptr in a Sequencer Section extension.

diff --git a/Engine/Plugins/MovieScene/SequencerScripting/Source/SequencerScripting/Private/ExtensionLibraries/MovieSceneSectionExtensions.cpp b/Engine/Plugins/MovieScene/SequencerScripting/Source/SequencerScripting/Private/ExtensionLibraries/MovieSceneSectionExtensions.cpp
--- a/Engine/Plugins/MovieScene/SequencerScripting/Source/SequencerScripting/Private/ExtensionLibraries/MovieSceneSectionExtensions.cpp
+++ b/Engine/Plugins/MovieScene/SequencerScripting/Source/SequencerScripting/Private/ExtensionLibraries/MovieSceneSectionExtensions.cpp
@@ -304,6 +304,12 @@ bool GetSubSectionChain(UMovieSceneSubSection* InSubSection, UMovieSceneSequence
 
 int32 UMovieSceneSectionExtensions::GetParentSequenceFrame(UMovieSceneSubSection* InSubSection, int32 InFrame, UMovieSceneSequence* ParentSequence)
 {
+	if (InSubSection == nullptr || ParentSequence == nullptr)
+	{
+		UE_LOG(LogMovieScene, Error, TEXT("UMovieSceneSectionExtension::GetParentSequenceFrame failed because either sub section or parent sequence is null! SubSection: %s ParentSequence: %s"), InSubSection ? *InSubSection->GetFullName() : TEXT("None"), ParentSequence ? *ParentSequence->GetFullName() : TEXT("None"));
+		return InFrame;
+	}
+
 	TArray<UMovieSceneSubSection*> SubSectionChain;
 	GetSubSectionChain(InSubSection, ParentSequence, SubSectionChain);
 		
