Sequencer Scripting: Output error message if channel is not found in SetValueInChannel

diff --git a/Engine/Plugins/MovieScene/SequencerScripting/Source/SequencerScripting/Private/KeysAndChannels/MovieSceneScriptingChannel.h b/Engine/Plugins/MovieScene/SequencerScripting/Source/SequencerScripting/Private/KeysAndChannels/MovieSceneScriptingChannel.h
--- a/Engine/Plugins/MovieScene/SequencerScripting/Source/SequencerScripting/Private/KeysAndChannels/MovieSceneScriptingChannel.h
+++ b/Engine/Plugins/MovieScene/SequencerScripting/Source/SequencerScripting/Private/KeysAndChannels/MovieSceneScriptingChannel.h
@@ -314,14 +314,16 @@ struct TMovieSceneScriptingKey
 		if (Channel)
 		{
 			using namespace MovieScene;
-			if(!AssignValue(Channel, KeyHandle, InNewValue))
+			if (!AssignValue(Channel, KeyHandle, InNewValue))
 			{
 				FFrame::KismetExecutionMessage(TEXT("Invalid KeyIndex for MovieSceneScriptingKey, failed to set value. Did you forget to create the key through the channel?"), ELogVerbosity::Error);
 				return;
 			}
 		}
-
-		FFrame::KismetExecutionMessage(TEXT("Invalid ChannelHandle for MovieSceneScriptingKey, failed to set value. Did you forget to create the key through the channel?"), ELogVerbosity::Error);
+		else
+		{
+			FFrame::KismetExecutionMessage(TEXT("Invalid ChannelHandle for MovieSceneScriptingKey, failed to set value. Did you forget to create the key through the channel?"), ELogVerbosity::Error);
+		}
 	}
 
 public:
