Fix CIS

diff --git a/Engine/Plugins/MovieScene/SequencerScripting/Source/SequencerScriptingEditor/Private/SequencerTools.cpp b/Engine/Plugins/MovieScene/SequencerScripting/Source/SequencerScriptingEditor/Private/SequencerTools.cpp
--- a/Engine/Plugins/MovieScene/SequencerScripting/Source/SequencerScriptingEditor/Private/SequencerTools.cpp
+++ b/Engine/Plugins/MovieScene/SequencerScripting/Source/SequencerScriptingEditor/Private/SequencerTools.cpp
@@ -159,7 +159,7 @@ void GetCameras(FbxNode* Parent, TArray<FbxCamera*>& Cameras)
 }
 
 
-TArray<FGuid> AddActors(UWorld* World, UMovieSceneSequence* InSequence, UMovieScene* InMovieScene, IMovieScenePlayer* Player, const FMovieSceneSequenceIDRef& TemplateID,const TArray<TWeakObjectPtr<AActor> >& InActors)
+TArray<FGuid> AddActors(UWorld* World, UMovieSceneSequence* InSequence, UMovieScene* InMovieScene, IMovieScenePlayer* Player, FMovieSceneSequenceIDRef TemplateID,const TArray<TWeakObjectPtr<AActor> >& InActors)
 {
 	TArray<FGuid> PossessableGuids;
 
@@ -193,7 +193,7 @@ TArray<FGuid> AddActors(UWorld* World, UMovieSceneSequence* InSequence, UMovieSc
 }
 
 
-void ImportFBXCamera(UnFbx::FFbxImporter* FbxImporter, UWorld* World, UMovieSceneSequence* Sequence, UMovieScene* InMovieScene, IMovieScenePlayer* Player, const FMovieSceneSequenceIDRef& TemplateID, TMap<FGuid, FString>& InObjectBindingMap, bool bMatchByNameOnly, bool bCreateCameras)
+void ImportFBXCamera(UnFbx::FFbxImporter* FbxImporter, UWorld* World, UMovieSceneSequence* Sequence, UMovieScene* InMovieScene, IMovieScenePlayer* Player, FMovieSceneSequenceIDRef TemplateID, TMap<FGuid, FString>& InObjectBindingMap, bool bMatchByNameOnly, bool bCreateCameras)
 {
 	if (bCreateCameras)
 	{
diff --git a/Engine/Source/Editor/MovieSceneTools/Private/MovieSceneToolHelpers.cpp b/Engine/Source/Editor/MovieSceneTools/Private/MovieSceneToolHelpers.cpp
--- a/Engine/Source/Editor/MovieSceneTools/Private/MovieSceneToolHelpers.cpp
+++ b/Engine/Source/Editor/MovieSceneTools/Private/MovieSceneToolHelpers.cpp
@@ -2074,7 +2074,7 @@ bool MovieSceneToolHelpers::ExportFBX(UWorld* World, UMovieScene* MovieScene, IM
 		if (Actor != NULL)
 		{
 			FGuid ExistingGuid = Player->FindObjectId(*Actor, MovieSceneSequenceID::Root);
-			if (ExistingGuid.IsValid() && !bSelectedOnly || Bindings.Contains(ExistingGuid))
+			if (ExistingGuid.IsValid() && (!bSelectedOnly || Bindings.Contains(ExistingGuid)))
 			{
 				ActorToExport.Add(Actor);
 			}
