Sequencer: Fix crash caused by getting Start/EndFrame when it's not set. Check HasStart/EndFrame before getting values.

diff --git a/Engine/Plugins/MovieScene/SequencerScripting/Source/SequencerScripting/Private/ExtensionLibraries/MovieSceneSectionExtensions.cpp b/Engine/Plugins/MovieScene/SequencerScripting/Source/SequencerScripting/Private/ExtensionLibraries/MovieSceneSectionExtensions.cpp
--- a/Engine/Plugins/MovieScene/SequencerScripting/Source/SequencerScripting/Private/ExtensionLibraries/MovieSceneSectionExtensions.cpp
+++ b/Engine/Plugins/MovieScene/SequencerScripting/Source/SequencerScripting/Private/ExtensionLibraries/MovieSceneSectionExtensions.cpp
@@ -26,6 +26,12 @@ FSequencerScriptingRange UMovieSceneSectionExtensions::GetRange(UMovieSceneSecti
 
 int32 UMovieSceneSectionExtensions::GetStartFrame(UMovieSceneSection* Section)
 {
+	if (!Section->HasStartFrame())
+	{
+		FFrame::KismetExecutionMessage(TEXT("Section does not have a start frame"), ELogVerbosity::Error);
+		return -1;
+	}
+
 	UMovieScene* MovieScene = Section->GetTypedOuter<UMovieScene>();
 	if (MovieScene)
 	{
@@ -40,6 +46,12 @@ int32 UMovieSceneSectionExtensions::GetStartFrame(UMovieSceneSection* Section)
 
 float UMovieSceneSectionExtensions::GetStartFrameSeconds(UMovieSceneSection* Section)
 {
+	if (!Section->HasStartFrame())
+	{
+		FFrame::KismetExecutionMessage(TEXT("Section does not have a start frame"), ELogVerbosity::Error);
+		return -1.f;
+	}
+
 	UMovieScene* MovieScene = Section->GetTypedOuter<UMovieScene>();
 	if (MovieScene)
 	{
@@ -52,6 +64,12 @@ float UMovieSceneSectionExtensions::GetStartFrameSeconds(UMovieSceneSection* Sec
 
 int32 UMovieSceneSectionExtensions::GetEndFrame(UMovieSceneSection* Section)
 {
+	if (!Section->HasEndFrame())
+	{
+		FFrame::KismetExecutionMessage(TEXT("Section does not have an end frame"), ELogVerbosity::Error);
+		return -1;
+	}
+
 	UMovieScene* MovieScene = Section->GetTypedOuter<UMovieScene>();
 	if (MovieScene)
 	{
@@ -66,6 +84,12 @@ int32 UMovieSceneSectionExtensions::GetEndFrame(UMovieSceneSection* Section)
 
 float UMovieSceneSectionExtensions::GetEndFrameSeconds(UMovieSceneSection* Section)
 {
+	if (!Section->HasEndFrame())
+	{
+		FFrame::KismetExecutionMessage(TEXT("Section does not have an end frame"), ELogVerbosity::Error);
+		return -1.f;
+	}
+
 	UMovieScene* MovieScene = Section->GetTypedOuter<UMovieScene>();
 	if (MovieScene)
 	{
