Sequencer Scripting: Add MakeBindingID to generate an FMovieSceneObjectBindingID from a given master sequence and a FSequencerBindingProxy. This can be used to make a binding to an object in a subsequence

diff --git a/Engine/Plugins/MovieScene/SequencerScripting/Source/SequencerScripting/Private/ExtensionLibraries/MovieSceneSequenceExtensions.cpp b/Engine/Plugins/MovieScene/SequencerScripting/Source/SequencerScripting/Private/ExtensionLibraries/MovieSceneSequenceExtensions.cpp
--- a/Engine/Plugins/MovieScene/SequencerScripting/Source/SequencerScripting/Private/ExtensionLibraries/MovieSceneSequenceExtensions.cpp
+++ b/Engine/Plugins/MovieScene/SequencerScripting/Source/SequencerScripting/Private/ExtensionLibraries/MovieSceneSequenceExtensions.cpp
@@ -5,6 +5,7 @@
 #include "MovieScene.h"
 #include "MovieSceneFolder.h"
 #include "Algo/Find.h"
+#include "Compilation/MovieSceneCompiler.h"
 
 TArray<UMovieSceneTrack*> UMovieSceneSequenceExtensions::FilterTracks(TArrayView<UMovieSceneTrack* const> InTracks, UClass* DesiredClass, bool bExactMatch)
 {
@@ -469,6 +470,28 @@ TArray<UObject*> UMovieSceneSequenceExtensions::LocateBoundObjects(UMovieSceneSe
 	return Result;
 }
 
+FMovieSceneObjectBindingID UMovieSceneSequenceExtensions::MakeBindingID(UMovieSceneSequence* MasterSequence, const FSequencerBindingProxy& InBinding)
+{
+	FMovieSceneSequenceID SequenceID = MovieSceneSequenceID::Root;
+
+	FMovieSceneSequenceHierarchy SequenceHierarchyCache;
+	FMovieSceneCompiler::CompileHierarchy(*MasterSequence, SequenceHierarchyCache);
+
+	for (const TTuple<FMovieSceneSequenceID, FMovieSceneSubSequenceData>& Pair : SequenceHierarchyCache.AllSubSequenceData())
+	{
+		if (UMovieSceneSequence* SubSequence = Pair.Value.GetSequence())
+		{
+			if (SubSequence == InBinding.Sequence)
+			{
+				SequenceID = Pair.Key;
+				break;
+			}
+		}
+	}
+
+	return FMovieSceneObjectBindingID(InBinding.BindingID, SequenceID);
+}
+
 TArray<UMovieSceneFolder*> UMovieSceneSequenceExtensions::GetRootFoldersInSequence(UMovieSceneSequence* Sequence)
 {
 	TArray<UMovieSceneFolder*> Result;
diff --git a/Engine/Plugins/MovieScene/SequencerScripting/Source/SequencerScripting/Public/ExtensionLibraries/MovieSceneSequenceExtensions.h b/Engine/Plugins/MovieScene/SequencerScripting/Source/SequencerScripting/Public/ExtensionLibraries/MovieSceneSequenceExtensions.h
--- a/Engine/Plugins/MovieScene/SequencerScripting/Source/SequencerScripting/Public/ExtensionLibraries/MovieSceneSequenceExtensions.h
+++ b/Engine/Plugins/MovieScene/SequencerScripting/Source/SequencerScripting/Public/ExtensionLibraries/MovieSceneSequenceExtensions.h
@@ -361,6 +361,17 @@ public:
 	UFUNCTION(BlueprintCallable, Category="Sequence", meta=(ScriptMethod))
 	static TArray<UObject*> LocateBoundObjects(UMovieSceneSequence* Sequence, const FSequencerBindingProxy& InBinding, UObject* Context);
 
+
+	/**
+	 * Make a binding id for the given binding in this sequence
+	 *
+	 * @param MasterSequence  The master sequence that contains the sequence
+	 * @param Binding The binding proxy to generate the binding id from
+	 * @return The new object binding id
+	 */
+	UFUNCTION(BlueprintCallable, Category = "Sequence", meta = (ScriptMethod))
+	static FMovieSceneObjectBindingID MakeBindingID(UMovieSceneSequence* MasterSequence, const FSequencerBindingProxy& InBinding);
+
 	/**
 	 * Get the root folders in the provided sequence
 	 *
