Sequencer: Expose Set and Get Object PropertyClass

diff --git a/Engine/Plugins/MovieScene/SequencerScripting/Source/SequencerScripting/Private/ExtensionLibraries/MovieScenePropertyTrackExtensions.cpp b/Engine/Plugins/MovieScene/SequencerScripting/Source/SequencerScripting/Private/ExtensionLibraries/MovieScenePropertyTrackExtensions.cpp
--- a/Engine/Plugins/MovieScene/SequencerScripting/Source/SequencerScripting/Private/ExtensionLibraries/MovieScenePropertyTrackExtensions.cpp
+++ b/Engine/Plugins/MovieScene/SequencerScripting/Source/SequencerScripting/Private/ExtensionLibraries/MovieScenePropertyTrackExtensions.cpp
@@ -2,6 +2,7 @@
 
 #include "ExtensionLibraries/MovieScenePropertyTrackExtensions.h"
 #include "Tracks/MovieScenePropertyTrack.h"
+#include "Tracks/MovieSceneObjectPropertyTrack.h"
 
 void UMovieScenePropertyTrackExtensions::SetPropertyNameAndPath(UMovieScenePropertyTrack* Track, const FName& InPropertyName, const FString& InPropertyPath)
 {
@@ -26,3 +27,15 @@ FName UMovieScenePropertyTrackExtensions::GetUniqueTrackName(UMovieSceneProperty
 	return NAME_None;
 #endif
 }
+
+void UMovieScenePropertyTrackExtensions::SetObjectPropertyClass(UMovieSceneObjectPropertyTrack* Track, UClass* PropertyClass)
+{
+	Track->PropertyClass = PropertyClass;
+}
+
+UClass* UMovieScenePropertyTrackExtensions::GetObjectPropertyClass(UMovieSceneObjectPropertyTrack* Track)
+{
+	return Track->PropertyClass;
+}
+
+
diff --git a/Engine/Plugins/MovieScene/SequencerScripting/Source/SequencerScripting/Public/ExtensionLibraries/MovieScenePropertyTrackExtensions.h b/Engine/Plugins/MovieScene/SequencerScripting/Source/SequencerScripting/Public/ExtensionLibraries/MovieScenePropertyTrackExtensions.h
--- a/Engine/Plugins/MovieScene/SequencerScripting/Source/SequencerScripting/Public/ExtensionLibraries/MovieScenePropertyTrackExtensions.h
+++ b/Engine/Plugins/MovieScene/SequencerScripting/Source/SequencerScripting/Public/ExtensionLibraries/MovieScenePropertyTrackExtensions.h
@@ -10,6 +10,7 @@
 
 
 class UMovieScenePropertyTrack;
+class UMovieSceneObjectPropertyTrack;
 
 /**
  * Function library containing methods that should be hoisted onto UMovieScenePropertyTrack for scripting
@@ -57,4 +58,22 @@ public:
 	 */
 	UFUNCTION(BlueprintCallable, Category="Track", meta=(ScriptMethod))
 	static FName GetUniqueTrackName(UMovieScenePropertyTrack* Track);
+
+	/**
+	 * Set this object property track's property class
+	 *
+	 * @param Track        The track to use
+	 * @param PropertyClass The property class to set
+	 */
+	UFUNCTION(BlueprintCallable, Category = "Track", meta = (ScriptMethod))
+	static void SetObjectPropertyClass(UMovieSceneObjectPropertyTrack* Track, UClass* PropertyClass);
+
+	/**
+	 * Get this object property track's property class
+	 *
+	 * @param Track        The track to use
+	 * @return The property class for this object property track
+	 */
+	UFUNCTION(BlueprintCallable, Category = "Track", meta = (ScriptMethod))
+	static UClass* GetObjectPropertyClass(UMovieSceneObjectPropertyTrack* Track);
 };
