Sequencer Scripting: Added the ability to create bindings in root or local space. Camera cut bindings in a subsequence should be created in local space so that they can resolve when opened on its own. Updated the master sequence example to create a binding in local space

diff --git a/Engine/Plugins/MovieScene/SequencerScripting/Source/SequencerScripting/Private/ExtensionLibraries/MovieSceneSequenceExtensions.cpp b/Engine/Plugins/MovieScene/SequencerScripting/Source/SequencerScripting/Private/ExtensionLibraries/MovieSceneSequenceExtensions.cpp
--- a/Engine/Plugins/MovieScene/SequencerScripting/Source/SequencerScripting/Private/ExtensionLibraries/MovieSceneSequenceExtensions.cpp
+++ b/Engine/Plugins/MovieScene/SequencerScripting/Source/SequencerScripting/Private/ExtensionLibraries/MovieSceneSequenceExtensions.cpp
@@ -470,7 +470,7 @@ TArray<UObject*> UMovieSceneSequenceExtensions::LocateBoundObjects(UMovieSceneSe
 	return Result;
 }
 
-FMovieSceneObjectBindingID UMovieSceneSequenceExtensions::MakeBindingID(UMovieSceneSequence* MasterSequence, const FSequencerBindingProxy& InBinding)
+FMovieSceneObjectBindingID UMovieSceneSequenceExtensions::MakeBindingID(UMovieSceneSequence* MasterSequence, const FSequencerBindingProxy& InBinding, EMovieSceneObjectBindingSpace Space)
 {
 	FMovieSceneSequenceID SequenceID = MovieSceneSequenceID::Root;
 
@@ -489,7 +489,7 @@ FMovieSceneObjectBindingID UMovieSceneSequenceExtensions::MakeBindingID(UMovieSc
 		}
 	}
 
-	return FMovieSceneObjectBindingID(InBinding.BindingID, SequenceID);
+	return FMovieSceneObjectBindingID(InBinding.BindingID, SequenceID, Space);
 }
 
 TArray<UMovieSceneFolder*> UMovieSceneSequenceExtensions::GetRootFoldersInSequence(UMovieSceneSequence* Sequence)
diff --git a/Engine/Plugins/MovieScene/SequencerScripting/Source/SequencerScripting/Public/ExtensionLibraries/MovieSceneSequenceExtensions.h b/Engine/Plugins/MovieScene/SequencerScripting/Source/SequencerScripting/Public/ExtensionLibraries/MovieSceneSequenceExtensions.h
--- a/Engine/Plugins/MovieScene/SequencerScripting/Source/SequencerScripting/Public/ExtensionLibraries/MovieSceneSequenceExtensions.h
+++ b/Engine/Plugins/MovieScene/SequencerScripting/Source/SequencerScripting/Public/ExtensionLibraries/MovieSceneSequenceExtensions.h
@@ -373,10 +373,11 @@ public:
 	 *
 	 * @param MasterSequence  The master sequence that contains the sequence
 	 * @param Binding The binding proxy to generate the binding id from
+	 * @param Space The object binding space to resolve from (Root or Local)
 	 * @return The new object binding id
 	 */
 	UFUNCTION(BlueprintCallable, Category = "Sequence", meta = (ScriptMethod))
-	static FMovieSceneObjectBindingID MakeBindingID(UMovieSceneSequence* MasterSequence, const FSequencerBindingProxy& InBinding);
+	static FMovieSceneObjectBindingID MakeBindingID(UMovieSceneSequence* MasterSequence, const FSequencerBindingProxy& InBinding, EMovieSceneObjectBindingSpace Space = EMovieSceneObjectBindingSpace::Root);
 
 	/**
 	 * Get the root folders in the provided sequence
