Sequencer Scripting: Added SetReadOnly, IsReadOnly

diff --git a/Engine/Plugins/MovieScene/SequencerScripting/Source/SequencerScripting/Private/ExtensionLibraries/MovieSceneSequenceExtensions.cpp b/Engine/Plugins/MovieScene/SequencerScripting/Source/SequencerScripting/Private/ExtensionLibraries/MovieSceneSequenceExtensions.cpp
--- a/Engine/Plugins/MovieScene/SequencerScripting/Source/SequencerScripting/Private/ExtensionLibraries/MovieSceneSequenceExtensions.cpp
+++ b/Engine/Plugins/MovieScene/SequencerScripting/Source/SequencerScripting/Private/ExtensionLibraries/MovieSceneSequenceExtensions.cpp
@@ -593,4 +593,24 @@ int32 UMovieSceneSequenceExtensions::FindNextMarkedFrame(UMovieSceneSequence* Se
 		return MovieScene->FindNextMarkedFrame(InFrameNumber, bForward);
 	}
 	return INDEX_NONE;
-}
\ No newline at end of file
+}
+
+void UMovieSceneSequenceExtensions::SetReadOnly(UMovieSceneSequence* Sequence, bool bInReadOnly)
+{
+	UMovieScene* MovieScene = Sequence->GetMovieScene();
+	if (MovieScene)
+	{
+		MovieScene->SetReadOnly(bInReadOnly);
+	}
+}
+
+bool UMovieSceneSequenceExtensions::IsReadOnly(UMovieSceneSequence* Sequence)
+{
+	UMovieScene* MovieScene = Sequence->GetMovieScene();
+	if (MovieScene)
+	{
+		return MovieScene->IsReadOnly();
+	}
+
+	return false;
+}
diff --git a/Engine/Plugins/MovieScene/SequencerScripting/Source/SequencerScripting/Public/ExtensionLibraries/MovieSceneSequenceExtensions.h b/Engine/Plugins/MovieScene/SequencerScripting/Source/SequencerScripting/Public/ExtensionLibraries/MovieSceneSequenceExtensions.h
--- a/Engine/Plugins/MovieScene/SequencerScripting/Source/SequencerScripting/Public/ExtensionLibraries/MovieSceneSequenceExtensions.h
+++ b/Engine/Plugins/MovieScene/SequencerScripting/Source/SequencerScripting/Public/ExtensionLibraries/MovieSceneSequenceExtensions.h
@@ -466,4 +466,20 @@ public:
 	 */
 	UFUNCTION(BlueprintCallable, Category = "Sequence", meta = (ScriptMethod))
 	static int32 FindNextMarkedFrame(UMovieSceneSequence* Sequence, FFrameNumber InFrameNumber, bool bForward);
+
+	/*
+	 * Set read only
+	 *
+	 * @bInReadOnly Whether the movie scene should be read only or not
+	 */
+	UFUNCTION(BlueprintCallable, Category = "Sequence", meta = (ScriptMethod))
+	static void SetReadOnly(UMovieSceneSequence* Sequence, bool bInReadOnly);
+
+	/*
+	 * Is read only
+	 *
+	 * @return Whether the movie scene is read only or not
+	 */
+	UFUNCTION(BlueprintCallable, Category = "Sequence", meta = (ScriptMethod))
+	static bool IsReadOnly(UMovieSceneSequence* Sequence);
 };
\ No newline at end of file
