Sequencer Scripting: Fix static analysis

#rnx

diff --git a/Engine/Plugins/MovieScene/SequencerScripting/Source/SequencerScripting/Private/ExtensionLibraries/MovieSceneSequenceExtensions.cpp b/Engine/Plugins/MovieScene/SequencerScripting/Source/SequencerScripting/Private/ExtensionLibraries/MovieSceneSequenceExtensions.cpp
--- a/Engine/Plugins/MovieScene/SequencerScripting/Source/SequencerScripting/Private/ExtensionLibraries/MovieSceneSequenceExtensions.cpp
+++ b/Engine/Plugins/MovieScene/SequencerScripting/Source/SequencerScripting/Private/ExtensionLibraries/MovieSceneSequenceExtensions.cpp
@@ -37,11 +37,11 @@ TArray<UMovieSceneTrack*> UMovieSceneSequenceExtensions::GetMasterTracks(UMovieS
 	if (MovieScene)
 	{
 		Tracks = MovieScene->GetMasterTracks();
-	}
 
-	if (UMovieSceneTrack* CameraCutTrack = MovieScene->GetCameraCutTrack())
-	{
-		Tracks.Add(CameraCutTrack);
+		if (UMovieSceneTrack* CameraCutTrack = MovieScene->GetCameraCutTrack())
+		{
+			Tracks.Add(CameraCutTrack);
+		}
 	}
 
 	return Tracks;
