Sequencer Scripting: Fix so that camera cut tracks are added as a camera cut track, rather than as a master track

#rnx

diff --git a/Engine/Plugins/MovieScene/SequencerScripting/Source/SequencerScripting/Private/ExtensionLibraries/MovieSceneSequenceExtensions.cpp b/Engine/Plugins/MovieScene/SequencerScripting/Source/SequencerScripting/Private/ExtensionLibraries/MovieSceneSequenceExtensions.cpp
--- a/Engine/Plugins/MovieScene/SequencerScripting/Source/SequencerScripting/Private/ExtensionLibraries/MovieSceneSequenceExtensions.cpp
+++ b/Engine/Plugins/MovieScene/SequencerScripting/Source/SequencerScripting/Private/ExtensionLibraries/MovieSceneSequenceExtensions.cpp
@@ -1,6 +1,7 @@
 // Copyright Epic Games, Inc. All Rights Reserved.
 
 #include "ExtensionLibraries/MovieSceneSequenceExtensions.h"
+#include "Tracks/MovieSceneCameraCutTrack.h"
 #include "MovieSceneSequence.h"
 #include "MovieScene.h"
 #include "MovieSceneFolder.h"
@@ -99,10 +100,13 @@ UMovieSceneTrack* UMovieSceneSequenceExtensions::AddMasterTrack(UMovieSceneSeque
 	UMovieScene* MovieScene = GetMovieScene(Sequence);
 	if (MovieScene)
 	{
-		UMovieSceneTrack* NewTrack = MovieScene->AddMasterTrack(TrackType);
-		if (NewTrack)
+		if (TrackType->IsChildOf(UMovieSceneCameraCutTrack::StaticClass()))
 		{
-			return NewTrack;
+			return MovieScene->AddCameraCutTrack(TrackType);
+		}
+		else
+		{
+			return MovieScene->AddMasterTrack(TrackType);
 		}
 	}
 
