Sequencer: Set end frame from scripting should set an exclusive bounds

diff --git a/Engine/Plugins/MovieScene/SequencerScripting/Source/SequencerScripting/Private/ExtensionLibraries/MovieSceneSectionExtensions.cpp b/Engine/Plugins/MovieScene/SequencerScripting/Source/SequencerScripting/Private/ExtensionLibraries/MovieSceneSectionExtensions.cpp
--- a/Engine/Plugins/MovieScene/SequencerScripting/Source/SequencerScripting/Private/ExtensionLibraries/MovieSceneSectionExtensions.cpp
+++ b/Engine/Plugins/MovieScene/SequencerScripting/Source/SequencerScripting/Private/ExtensionLibraries/MovieSceneSectionExtensions.cpp
@@ -194,7 +194,7 @@ void UMovieSceneSectionExtensions::SetEndFrame(UMovieSceneSection* Section, int3
 	{
 		FFrameRate DisplayRate = MovieScene->GetDisplayRate();
 
-		Section->SetEndFrame(ConvertFrameTime(EndFrame, DisplayRate, MovieScene->GetTickResolution()).FrameNumber);
+		Section->SetEndFrame(TRangeBound<FFrameNumber>::Exclusive(ConvertFrameTime(EndFrame, DisplayRate, MovieScene->GetTickResolution()).FrameNumber));
 	}
 }
 
@@ -203,7 +203,7 @@ void UMovieSceneSectionExtensions::SetEndFrameSeconds(UMovieSceneSection* Sectio
 	UMovieScene* MovieScene = Section->GetTypedOuter<UMovieScene>();
 	if (MovieScene)
 	{
-		Section->SetEndFrame((EndTime * MovieScene->GetTickResolution()).RoundToFrame());
+		Section->SetEndFrame(TRangeBound<FFrameNumber>::Exclusive((EndTime * MovieScene->GetTickResolution()).RoundToFrame()));
 	}
 }
 
