Minor change to fix seemingly-spurious link error in SequencerScripting module on older compiler versions.

#rnx

diff --git a/Engine/Plugins/MovieScene/SequencerScripting/Source/SequencerScripting/Private/KeysAndChannels/MovieSceneScriptingChannel.h b/Engine/Plugins/MovieScene/SequencerScripting/Source/SequencerScripting/Private/KeysAndChannels/MovieSceneScriptingChannel.h
--- a/Engine/Plugins/MovieScene/SequencerScripting/Source/SequencerScripting/Private/KeysAndChannels/MovieSceneScriptingChannel.h
+++ b/Engine/Plugins/MovieScene/SequencerScripting/Source/SequencerScripting/Private/KeysAndChannels/MovieSceneScriptingChannel.h
@@ -87,7 +87,8 @@ struct TMovieSceneScriptingChannel
 		ChannelType* Channel = ChannelHandle.Get();
 		if (Channel)
 		{
-			Channel->DeleteKeys(MakeArrayView(&Key->KeyHandle, 1));
+			TArrayView<FKeyHandle, int> KeyArray = TArrayView<FKeyHandle, int>(&Key->KeyHandle, 1);
+			Channel->DeleteKeys(KeyArray);
 			return;
 		}
 
