Sequencer: Add backwards-compatible scripting wrappers that handle double-channels as float.

Some new methods will provide the proper double channels.

#okforgitub public

diff --git a/Engine/Plugins/MovieScene/SequencerScripting/Source/SequencerScripting/Private/ExtensionLibraries/MovieSceneSectionExtensions.cpp b/Engine/Plugins/MovieScene/SequencerScripting/Source/SequencerScripting/Private/ExtensionLibraries/MovieSceneSectionExtensions.cpp
--- a/Engine/Plugins/MovieScene/SequencerScripting/Source/SequencerScripting/Private/ExtensionLibraries/MovieSceneSectionExtensions.cpp
+++ b/Engine/Plugins/MovieScene/SequencerScripting/Source/SequencerScripting/Private/ExtensionLibraries/MovieSceneSectionExtensions.cpp
@@ -12,6 +12,7 @@
 #include "KeysAndChannels/MovieSceneScriptingByte.h"
 #include "KeysAndChannels/MovieSceneScriptingInteger.h"
 #include "KeysAndChannels/MovieSceneScriptingFloat.h"
+#include "KeysAndChannels/MovieSceneScriptingDouble.h"
 #include "KeysAndChannels/MovieSceneScriptingString.h"
 #include "KeysAndChannels/MovieSceneScriptingEvent.h"
 #include "KeysAndChannels/MovieSceneScriptingActorReference.h"
@@ -240,6 +241,22 @@ void UMovieSceneSectionExtensions::SetEndFrameBounded(UMovieSceneSection* Sectio
 	}
 }
 
+template<typename ChannelType, typename ScriptingChannelType>
+void SetScriptingChannelHandle(ScriptingChannelType* ScriptingChannel, FMovieSceneChannelProxy& ChannelProxy, int32 ChannelIndex)
+{
+	ScriptingChannel->ChannelHandle = ChannelProxy.MakeHandle<ChannelType>(ChannelIndex);
+}
+template<>
+void SetScriptingChannelHandle<FMovieSceneFloatChannel, UMovieSceneScriptingFloatChannel>(UMovieSceneScriptingFloatChannel* ScriptingChannel, FMovieSceneChannelProxy& ChannelProxy, int32 ChannelIndex)
+{
+	ScriptingChannel->FloatChannelHandle = ChannelProxy.MakeHandle<FMovieSceneFloatChannel>(ChannelIndex);
+}
+template<>
+void SetScriptingChannelHandle<FMovieSceneDoubleChannel, UMovieSceneScriptingFloatChannel>(UMovieSceneScriptingFloatChannel* ScriptingChannel, FMovieSceneChannelProxy& ChannelProxy, int32 ChannelIndex)
+{
+	ScriptingChannel->DoubleChannelHandle = ChannelProxy.MakeHandle<FMovieSceneDoubleChannel>(ChannelIndex);
+}
+
 template<typename ChannelType, typename ScriptingChannelType>
 void GetScriptingChannelsForChannel(FMovieSceneChannelProxy& ChannelProxy,TWeakObjectPtr<UMovieSceneSequence> Sequence, TWeakObjectPtr<UMovieSceneSection> Section, TArray<UMovieSceneScriptingChannel*>& OutChannels)
 {
@@ -255,7 +272,7 @@ void GetScriptingChannelsForChannel(FMovieSceneChannelProxy& ChannelProxy,TWeakO
 				const FName ChannelName = MetaData[Index].Name;
 				FName UniqueChannelName = MakeUniqueObjectName(GetTransientPackage(), ScriptingChannelType::StaticClass(), ChannelName);
 				ScriptingChannelType* ScriptingChannel = NewObject<ScriptingChannelType>(GetTransientPackage(), UniqueChannelName);
-				ScriptingChannel->ChannelHandle = ChannelProxy.MakeHandle<ChannelType>(Index);
+				SetScriptingChannelHandle<ChannelType, ScriptingChannelType>(ScriptingChannel, ChannelProxy, Index);
 				ScriptingChannel->OwningSection = Section;
 				ScriptingChannel->OwningSequence = Sequence;
 				ScriptingChannel->ChannelName = ChannelName;
@@ -279,11 +296,41 @@ TArray<UMovieSceneScriptingChannel*> UMovieSceneSectionExtensions::GetChannels(U
 
 	// ToDo: This isn't a great way of collecting all of the channel types and instantiating the correct UObject for it
 	// but given that we're already hard-coding the UObject implementations... We currently support the following channel
-	// types: Bool, Byte, Float, Integer, String, Event, ActorReference
+	// types: Bool, Byte, Float, Double, Integer, String, Event, ActorReference
+	TWeakObjectPtr<UMovieSceneSequence> Sequence = Section->GetTypedOuter<UMovieSceneSequence>();
+	GetScriptingChannelsForChannel<FMovieSceneBoolChannel, UMovieSceneScriptingBoolChannel>(ChannelProxy, Sequence, Section, Channels);
+	GetScriptingChannelsForChannel<FMovieSceneByteChannel, UMovieSceneScriptingByteChannel>(ChannelProxy, Sequence, Section, Channels);
+	GetScriptingChannelsForChannel<FMovieSceneFloatChannel, UMovieSceneScriptingFloatChannel>(ChannelProxy, Sequence, Section, Channels);
+	GetScriptingChannelsForChannel<FMovieSceneIntegerChannel, UMovieSceneScriptingIntegerChannel>(ChannelProxy, Sequence, Section, Channels);
+	GetScriptingChannelsForChannel<FMovieSceneStringChannel, UMovieSceneScriptingStringChannel>(ChannelProxy, Sequence, Section, Channels);
+	GetScriptingChannelsForChannel<FMovieSceneEventChannel, UMovieSceneScriptingEventChannel>(ChannelProxy, Sequence, Section, Channels);
+	GetScriptingChannelsForChannel<FMovieSceneActorReferenceData, UMovieSceneScriptingActorReferenceChannel>(ChannelProxy, Sequence, Section, Channels);
+	GetScriptingChannelsForChannel<FMovieSceneObjectPathChannel, UMovieSceneScriptingObjectPathChannel>(ChannelProxy, Sequence, Section, Channels);
+
+	// Wrap double channels around float scripting channels.
+	GetScriptingChannelsForChannel<FMovieSceneDoubleChannel, UMovieSceneScriptingFloatChannel>(ChannelProxy, Sequence, Section, Channels);
+
+	return Channels;
+}
+
+TArray<UMovieSceneScriptingChannel*> UMovieSceneSectionExtensions::GetAllChannels(UMovieSceneSection* Section)
+{
+	TArray<UMovieSceneScriptingChannel*> Channels;
+	if (!Section)
+	{
+		FFrame::KismetExecutionMessage(TEXT("Cannot get channels for null section"), ELogVerbosity::Error);
+		return Channels;
+	}
+	FMovieSceneChannelProxy& ChannelProxy = Section->GetChannelProxy();
+
+	// ToDo: This isn't a great way of collecting all of the channel types and instantiating the correct UObject for it
+	// but given that we're already hard-coding the UObject implementations... We currently support the following channel
+	// types: Bool, Byte, Float, Double, Integer, String, Event, ActorReference
 	TWeakObjectPtr<UMovieSceneSequence> Sequence = Section->GetTypedOuter<UMovieSceneSequence>();
 	GetScriptingChannelsForChannel<FMovieSceneBoolChannel, UMovieSceneScriptingBoolChannel>(ChannelProxy, Sequence, Section, Channels);
 	GetScriptingChannelsForChannel<FMovieSceneByteChannel, UMovieSceneScriptingByteChannel>(ChannelProxy, Sequence, Section, Channels);
 	GetScriptingChannelsForChannel<FMovieSceneFloatChannel, UMovieSceneScriptingFloatChannel>(ChannelProxy, Sequence, Section, Channels);
+	GetScriptingChannelsForChannel<FMovieSceneDoubleChannel, UMovieSceneScriptingDoubleChannel>(ChannelProxy, Sequence, Section, Channels);
 	GetScriptingChannelsForChannel<FMovieSceneIntegerChannel, UMovieSceneScriptingIntegerChannel>(ChannelProxy, Sequence, Section, Channels);
 	GetScriptingChannelsForChannel<FMovieSceneStringChannel, UMovieSceneScriptingStringChannel>(ChannelProxy, Sequence, Section, Channels);
 	GetScriptingChannelsForChannel<FMovieSceneEventChannel, UMovieSceneScriptingEventChannel>(ChannelProxy, Sequence, Section, Channels);
@@ -310,6 +357,7 @@ TArray<UMovieSceneScriptingChannel*> UMovieSceneSectionExtensions::FindChannelsB
 	if (ChannelType == UMovieSceneScriptingBoolChannel::StaticClass())					{ GetScriptingChannelsForChannel<FMovieSceneBoolChannel, UMovieSceneScriptingBoolChannel>(ChannelProxy, Sequence, Section, Channels); }
 	else if (ChannelType == UMovieSceneScriptingByteChannel::StaticClass())				{ GetScriptingChannelsForChannel<FMovieSceneByteChannel, UMovieSceneScriptingByteChannel>(ChannelProxy, Sequence, Section, Channels); }
 	else if (ChannelType == UMovieSceneScriptingFloatChannel::StaticClass())			{ GetScriptingChannelsForChannel<FMovieSceneFloatChannel, UMovieSceneScriptingFloatChannel>(ChannelProxy, Sequence, Section, Channels); }
+	else if (ChannelType == UMovieSceneScriptingDoubleChannel::StaticClass())			{ GetScriptingChannelsForChannel<FMovieSceneDoubleChannel, UMovieSceneScriptingFloatChannel>(ChannelProxy, Sequence, Section, Channels); }
 	else if (ChannelType == UMovieSceneScriptingIntegerChannel::StaticClass())			{ GetScriptingChannelsForChannel<FMovieSceneIntegerChannel, UMovieSceneScriptingIntegerChannel>(ChannelProxy, Sequence, Section, Channels); }
 	else if (ChannelType == UMovieSceneScriptingStringChannel::StaticClass())			{ GetScriptingChannelsForChannel<FMovieSceneStringChannel, UMovieSceneScriptingStringChannel>(ChannelProxy, Sequence, Section, Channels); }
 	else if (ChannelType == UMovieSceneScriptingEventChannel::StaticClass())			{ GetScriptingChannelsForChannel<FMovieSceneEventChannel, UMovieSceneScriptingEventChannel>(ChannelProxy, Sequence, Section, Channels); }
@@ -324,6 +372,37 @@ TArray<UMovieSceneScriptingChannel*> UMovieSceneSectionExtensions::FindChannelsB
 	return Channels;
 }
 
+TArray<UMovieSceneScriptingChannel*> UMovieSceneSectionExtensions::GetChannelsByType(UMovieSceneSection* Section, TSubclassOf<UMovieSceneScriptingChannel> ChannelType)
+{
+	TArray<UMovieSceneScriptingChannel*> Channels;
+	if (!Section)
+	{
+		FFrame::KismetExecutionMessage(TEXT("Cannot get channels for null section"), ELogVerbosity::Error);
+		return Channels;
+	}
+ 
+	FMovieSceneChannelProxy& ChannelProxy = Section->GetChannelProxy();
+ 
+	// ToDo: This needs to be done dynamically in the future but there's not a good way to do that right now with all of the SFINAE-based Templating driving these channels
+	TWeakObjectPtr<UMovieSceneSequence> Sequence = Section->GetTypedOuter<UMovieSceneSequence>();
+
+	if (ChannelType == UMovieSceneScriptingBoolChannel::StaticClass())					{ GetScriptingChannelsForChannel<FMovieSceneBoolChannel, UMovieSceneScriptingBoolChannel>(ChannelProxy, Sequence, Section, Channels); }
+	else if (ChannelType == UMovieSceneScriptingByteChannel::StaticClass())				{ GetScriptingChannelsForChannel<FMovieSceneByteChannel, UMovieSceneScriptingByteChannel>(ChannelProxy, Sequence, Section, Channels); }
+	else if (ChannelType == UMovieSceneScriptingFloatChannel::StaticClass())			{ GetScriptingChannelsForChannel<FMovieSceneFloatChannel, UMovieSceneScriptingFloatChannel>(ChannelProxy, Sequence, Section, Channels); }
+	else if (ChannelType == UMovieSceneScriptingDoubleChannel::StaticClass())			{ GetScriptingChannelsForChannel<FMovieSceneDoubleChannel, UMovieSceneScriptingDoubleChannel>(ChannelProxy, Sequence, Section, Channels); }
+	else if (ChannelType == UMovieSceneScriptingIntegerChannel::StaticClass())			{ GetScriptingChannelsForChannel<FMovieSceneIntegerChannel, UMovieSceneScriptingIntegerChannel>(ChannelProxy, Sequence, Section, Channels); }
+	else if (ChannelType == UMovieSceneScriptingStringChannel::StaticClass())			{ GetScriptingChannelsForChannel<FMovieSceneStringChannel, UMovieSceneScriptingStringChannel>(ChannelProxy, Sequence, Section, Channels); }
+	else if (ChannelType == UMovieSceneScriptingEventChannel::StaticClass())			{ GetScriptingChannelsForChannel<FMovieSceneEventChannel, UMovieSceneScriptingEventChannel>(ChannelProxy, Sequence, Section, Channels); }
+	else if (ChannelType == UMovieSceneScriptingActorReferenceChannel::StaticClass()) { GetScriptingChannelsForChannel<FMovieSceneActorReferenceData, UMovieSceneScriptingActorReferenceChannel>(ChannelProxy, Sequence, Section, Channels); }
+	else if (ChannelType == UMovieSceneScriptingObjectPathChannel::StaticClass()) { GetScriptingChannelsForChannel<FMovieSceneObjectPathChannel, UMovieSceneScriptingObjectPathChannel>(ChannelProxy, Sequence, Section,  Channels); }
+	else
+	{
+		FFrame::KismetExecutionMessage(TEXT("Unsupported ChannelType for FindChannelsByType!"), ELogVerbosity::Error);
+	}
+
+
+	return Channels;
+}
 
 bool GetSubSectionChain(UMovieSceneSubSection* InSubSection, UMovieSceneSequence* ParentSequence, TArray<UMovieSceneSubSection*>& SubSectionChain)
 {
diff --git a/Engine/Plugins/MovieScene/SequencerScripting/Source/SequencerScripting/Private/KeysAndChannels/MovieSceneScriptingActorReference.h b/Engine/Plugins/MovieScene/SequencerScripting/Source/SequencerScripting/Private/KeysAndChannels/MovieSceneScriptingActorReference.h
--- a/Engine/Plugins/MovieScene/SequencerScripting/Source/SequencerScripting/Private/KeysAndChannels/MovieSceneScriptingActorReference.h
+++ b/Engine/Plugins/MovieScene/SequencerScripting/Source/SequencerScripting/Private/KeysAndChannels/MovieSceneScriptingActorReference.h
@@ -65,9 +65,12 @@ public:
 };
 
 UCLASS(BlueprintType)
-class UMovieSceneScriptingActorReferenceChannel : public UMovieSceneScriptingChannel, public TMovieSceneScriptingChannel<FMovieSceneActorReferenceData, UMovieSceneScriptingActorReferenceKey, FMovieSceneActorReferenceKey>
+class UMovieSceneScriptingActorReferenceChannel : public UMovieSceneScriptingChannel
 {
 	GENERATED_BODY()
+
+	using Impl = TMovieSceneScriptingChannel<FMovieSceneActorReferenceData, UMovieSceneScriptingActorReferenceKey, FMovieSceneActorReferenceKey>;
+
 public:
 	/**
 	* Add a key to this channel. This initializes a new key and returns a reference to it.
@@ -81,7 +84,7 @@ public:
 	UMovieSceneScriptingActorReferenceKey* AddKey(const FFrameNumber InTime, FMovieSceneObjectBindingID NewValue, float SubFrame = 0.f, ESequenceTimeUnit TimeUnit = ESequenceTimeUnit::DisplayRate)
 	{
 		FMovieSceneActorReferenceKey ReferenceKey = FMovieSceneActorReferenceKey(NewValue);
-		return AddKeyInChannel(ChannelHandle, OwningSequence, OwningSection, InTime, ReferenceKey, SubFrame, TimeUnit, EMovieSceneKeyInterpolation::Auto);
+		return Impl::AddKeyInChannel(ChannelHandle, OwningSequence, OwningSection, InTime, ReferenceKey, SubFrame, TimeUnit, EMovieSceneKeyInterpolation::Auto);
 	}
 
 	/**
@@ -90,7 +93,7 @@ public:
 	UFUNCTION(BlueprintCallable, Category = "Sequencer|Keys", meta = (DisplayName = "Remove Key (Actor Reference)"))
 	virtual void RemoveKey(UMovieSceneScriptingKey* Key)
 	{
-		RemoveKeyFromChannel(ChannelHandle, Key);
+		Impl::RemoveKeyFromChannel(ChannelHandle, Key);
 	}
 
 	/**
@@ -101,7 +104,7 @@ public:
 	UFUNCTION(BlueprintCallable, Category = "Sequencer|Keys", meta = (DisplayName = "Get Keys (Actor Reference)"))
 	virtual TArray<UMovieSceneScriptingKey*> GetKeys() const override
 	{
-		return GetKeysInChannel(ChannelHandle, OwningSequence, OwningSection);
+		return Impl::GetKeysInChannel(ChannelHandle, OwningSequence, OwningSection);
 	}
 
 	/**
@@ -112,7 +115,7 @@ public:
 	void SetDefault(FMovieSceneObjectBindingID InDefaultValue)
 	{
 		FMovieSceneActorReferenceKey ReferenceKey = FMovieSceneActorReferenceKey(InDefaultValue);
-		SetDefaultInChannel(ChannelHandle, OwningSequence, OwningSection, ReferenceKey);
+		Impl::SetDefaultInChannel(ChannelHandle, OwningSequence, OwningSection, ReferenceKey);
 	}
 
 	/**
@@ -163,4 +166,4 @@ public:
 	TWeakObjectPtr<UMovieSceneSequence> OwningSequence;
 	TWeakObjectPtr<UMovieSceneSection> OwningSection;
 	TMovieSceneChannelHandle<FMovieSceneActorReferenceData> ChannelHandle;
-};
\ No newline at end of file
+};
diff --git a/Engine/Plugins/MovieScene/SequencerScripting/Source/SequencerScripting/Private/KeysAndChannels/MovieSceneScriptingBool.h b/Engine/Plugins/MovieScene/SequencerScripting/Source/SequencerScripting/Private/KeysAndChannels/MovieSceneScriptingBool.h
--- a/Engine/Plugins/MovieScene/SequencerScripting/Source/SequencerScripting/Private/KeysAndChannels/MovieSceneScriptingBool.h
+++ b/Engine/Plugins/MovieScene/SequencerScripting/Source/SequencerScripting/Private/KeysAndChannels/MovieSceneScriptingBool.h
@@ -63,9 +63,12 @@ public:
 };
 
 UCLASS(BlueprintType)
-class UMovieSceneScriptingBoolChannel : public UMovieSceneScriptingChannel, public TMovieSceneScriptingChannel<FMovieSceneBoolChannel, UMovieSceneScriptingBoolKey, bool>
+class UMovieSceneScriptingBoolChannel : public UMovieSceneScriptingChannel
 {
 	GENERATED_BODY()
+
+	using Impl = TMovieSceneScriptingChannel<FMovieSceneBoolChannel, UMovieSceneScriptingBoolKey, bool>;
+
 public:
 	/**
 	* Add a key to this channel. This initializes a new key and returns a reference to it.
@@ -78,7 +81,7 @@ public:
 	UFUNCTION(BlueprintCallable, Category = "Sequencer|Keys", meta=(DisplayName = "Add Key (Bool)"))
 	UMovieSceneScriptingBoolKey* AddKey(const FFrameNumber& InTime, bool NewValue, float SubFrame = 0.f, ESequenceTimeUnit TimeUnit = ESequenceTimeUnit::DisplayRate)
 	{
-		return AddKeyInChannel(ChannelHandle, OwningSequence, OwningSection, InTime, NewValue, SubFrame, TimeUnit, EMovieSceneKeyInterpolation::Auto);
+		return Impl::AddKeyInChannel(ChannelHandle, OwningSequence, OwningSection, InTime, NewValue, SubFrame, TimeUnit, EMovieSceneKeyInterpolation::Auto);
 	}
 
 	/**
@@ -87,7 +90,7 @@ public:
 	UFUNCTION(BlueprintCallable, Category = "Sequencer|Keys", meta=(DisplayName = "Remove Key (Bool)"))
 	virtual void RemoveKey(UMovieSceneScriptingKey* Key)
 	{
-		RemoveKeyFromChannel(ChannelHandle, Key);
+		Impl::RemoveKeyFromChannel(ChannelHandle, Key);
 	}
 
 	/**
@@ -98,7 +101,7 @@ public:
 	UFUNCTION(BlueprintCallable, Category = "Sequencer|Keys", meta = (DisplayName = "Get Keys (Bool)"))
 	virtual TArray<UMovieSceneScriptingKey*> GetKeys() const override
 	{
-		return GetKeysInChannel(ChannelHandle, OwningSequence, OwningSection);
+		return Impl::GetKeysInChannel(ChannelHandle, OwningSequence, OwningSection);
 	}
 
 	/**
@@ -118,7 +121,7 @@ public:
 	UFUNCTION(BlueprintCallable, Category = "Sequencer|Keys", meta = (DisplayName = "Evaluate Keys (Bool)"))
 	TArray<bool> EvaluateKeys(FSequencerScriptingRange Range, FFrameRate FrameRate) const
 	{
-		return EvaluateKeysInChannel(ChannelHandle, OwningSequence, Range, FrameRate);
+		return Impl::EvaluateKeysInChannel(ChannelHandle, OwningSequence, Range, FrameRate);
 	}
 
 	/**
@@ -129,7 +132,7 @@ public:
 	UFUNCTION(BlueprintCallable, Category = "Sequencer|Keys", meta = (DisplayName = "Compute Effective Range (Bool)"))
 	FSequencerScriptingRange ComputeEffectiveRange() const
 	{
-		return ComputeEffectiveRangeInChannel(ChannelHandle, OwningSequence);
+		return Impl::ComputeEffectiveRangeInChannel(ChannelHandle, OwningSequence);
 	}
 
 	/**
@@ -139,7 +142,7 @@ public:
 	UFUNCTION(BlueprintCallable, Category = "Sequencer|Keys", meta = (DisplayName = "Set Default (Bool)"))
 	void SetDefault(bool InDefaultValue)
 	{
-		SetDefaultInChannel(ChannelHandle, OwningSequence, OwningSection, InDefaultValue);
+		Impl::SetDefaultInChannel(ChannelHandle, OwningSequence, OwningSection, InDefaultValue);
 	}
 
 	/**
@@ -149,7 +152,7 @@ public:
 	UFUNCTION(BlueprintCallable, Category = "Sequencer|Keys", meta = (DisplayName = "Get Default (Bool)"))
 	bool GetDefault() const
 	{
-		TOptional<bool> DefaultValue = GetDefaultFromChannel(ChannelHandle);
+		TOptional<bool> DefaultValue = Impl::GetDefaultFromChannel(ChannelHandle);
 		return DefaultValue.IsSet() ? DefaultValue.GetValue() : false;
 	}
 
@@ -159,7 +162,7 @@ public:
 	UFUNCTION(BlueprintCallable, Category = "Sequencer|Keys", meta = (DisplayName = "Remove Default (Bool)"))
 	void RemoveDefault()
 	{
-		RemoveDefaultFromChannel(ChannelHandle);
+		Impl::RemoveDefaultFromChannel(ChannelHandle);
 	}
 
 	/**
@@ -168,7 +171,7 @@ public:
 	UFUNCTION(BlueprintCallable, Category = "Sequencer|Keys", meta = (DisplayName = "Has Default (Bool)"))
 	bool HasDefault() const
 	{
-		return GetDefaultFromChannel(ChannelHandle).IsSet();
+		return Impl::GetDefaultFromChannel(ChannelHandle).IsSet();
 	}
 public:
 	TWeakObjectPtr<UMovieSceneSequence> OwningSequence;
diff --git a/Engine/Plugins/MovieScene/SequencerScripting/Source/SequencerScripting/Private/KeysAndChannels/MovieSceneScriptingByte.h b/Engine/Plugins/MovieScene/SequencerScripting/Source/SequencerScripting/Private/KeysAndChannels/MovieSceneScriptingByte.h
--- a/Engine/Plugins/MovieScene/SequencerScripting/Source/SequencerScripting/Private/KeysAndChannels/MovieSceneScriptingByte.h
+++ b/Engine/Plugins/MovieScene/SequencerScripting/Source/SequencerScripting/Private/KeysAndChannels/MovieSceneScriptingByte.h
@@ -62,9 +62,12 @@ public:
 };
 
 UCLASS(BlueprintType)
-class UMovieSceneScriptingByteChannel : public UMovieSceneScriptingChannel, public TMovieSceneScriptingChannel<FMovieSceneByteChannel, UMovieSceneScriptingByteKey, uint8>
+class UMovieSceneScriptingByteChannel : public UMovieSceneScriptingChannel
 {
 	GENERATED_BODY()
+
+	using Impl = TMovieSceneScriptingChannel<FMovieSceneByteChannel, UMovieSceneScriptingByteKey, uint8>;
+	
 public:
 	/**
 	* Add a key to this channel. This initializes a new key and returns a reference to it.
@@ -78,7 +81,7 @@ public:
 	UFUNCTION(BlueprintCallable, Category = "Sequencer|Keys", meta = (DisplayName = "Add Key (Enum)"))
 	UMovieSceneScriptingByteKey* AddKey(const FFrameNumber& InTime, uint8 NewValue, float SubFrame = 0.f, ESequenceTimeUnit TimeUnit = ESequenceTimeUnit::DisplayRate, EMovieSceneKeyInterpolation InInterpolation = EMovieSceneKeyInterpolation::Auto)
 	{
-		return AddKeyInChannel(ChannelHandle, OwningSequence, OwningSection, InTime, NewValue, SubFrame, TimeUnit, InInterpolation);
+		return Impl::AddKeyInChannel(ChannelHandle, OwningSequence, OwningSection, InTime, NewValue, SubFrame, TimeUnit, InInterpolation);
 	}
 
 	/**
@@ -87,7 +90,7 @@ public:
 	UFUNCTION(BlueprintCallable, Category = "Sequencer|Keys", meta = (DisplayName = "Remove Key (Enum)"))
 	virtual void RemoveKey(UMovieSceneScriptingKey* Key)
 	{
-		RemoveKeyFromChannel(ChannelHandle, Key);
+		Impl::RemoveKeyFromChannel(ChannelHandle, Key);
 	}
 
 	/**
@@ -98,7 +101,7 @@ public:
 	UFUNCTION(BlueprintCallable, Category = "Sequencer|Keys", meta = (DisplayName = "Get Keys (Enum)"))
 	virtual TArray<UMovieSceneScriptingKey*> GetKeys() const override
 	{
-		return GetKeysInChannel(ChannelHandle, OwningSequence, OwningSection);
+		return Impl::GetKeysInChannel(ChannelHandle, OwningSequence, OwningSection);
 	}
 
 	/**
@@ -108,7 +111,7 @@ public:
 	UFUNCTION(BlueprintCallable, Category = "Sequencer|Keys", meta = (DisplayName = "Set Default (Enum)"))
 	void SetDefault(uint8 InDefaultValue)
 	{
-		SetDefaultInChannel(ChannelHandle, OwningSequence, OwningSection, InDefaultValue);
+		Impl::SetDefaultInChannel(ChannelHandle, OwningSequence, OwningSection, InDefaultValue);
 	}
 
 	/**
@@ -118,7 +121,7 @@ public:
 	UFUNCTION(BlueprintCallable, Category = "Sequencer|Keys", meta = (DisplayName = "Get Default (Enum)"))
 	uint8 GetDefault() const
 	{
-		TOptional<uint8> DefaultValue = GetDefaultFromChannel(ChannelHandle);
+		TOptional<uint8> DefaultValue = Impl::GetDefaultFromChannel(ChannelHandle);
 		return DefaultValue.IsSet() ? DefaultValue.GetValue() : 0;
 	}
 
@@ -128,7 +131,7 @@ public:
 	UFUNCTION(BlueprintCallable, Category = "Sequencer|Keys", meta = (DisplayName = "Remove Default (Enum)"))
 	void RemoveDefault()
 	{
-		RemoveDefaultFromChannel(ChannelHandle);
+		Impl::RemoveDefaultFromChannel(ChannelHandle);
 	}
 
 	/**
@@ -137,10 +140,10 @@ public:
 	UFUNCTION(BlueprintCallable, Category = "Sequencer|Keys", meta = (DisplayName = "Has Default (Enum)"))
 	bool HasDefault() const
 	{
-		return GetDefaultFromChannel(ChannelHandle).IsSet();
+		return Impl::GetDefaultFromChannel(ChannelHandle).IsSet();
 	}
 public:
 	TWeakObjectPtr<UMovieSceneSequence> OwningSequence;
 	TWeakObjectPtr<UMovieSceneSection> OwningSection;
 	TMovieSceneChannelHandle<FMovieSceneByteChannel> ChannelHandle;
-};
\ No newline at end of file
+};
diff --git a/Engine/Plugins/MovieScene/SequencerScripting/Source/SequencerScripting/Private/KeysAndChannels/MovieSceneScriptingChannel.h b/Engine/Plugins/MovieScene/SequencerScripting/Source/SequencerScripting/Private/KeysAndChannels/MovieSceneScriptingChannel.h
--- a/Engine/Plugins/MovieScene/SequencerScripting/Source/SequencerScripting/Private/KeysAndChannels/MovieSceneScriptingChannel.h
+++ b/Engine/Plugins/MovieScene/SequencerScripting/Source/SequencerScripting/Private/KeysAndChannels/MovieSceneScriptingChannel.h
@@ -55,7 +55,7 @@ public:
 template<typename ChannelType, typename ScriptingKeyType, typename ScriptingKeyValueType>
 struct TMovieSceneScriptingChannel
 {
-	ScriptingKeyType* AddKeyInChannel(TMovieSceneChannelHandle<ChannelType> ChannelHandle, TWeakObjectPtr<UMovieSceneSequence> Sequence, 
+	static ScriptingKeyType* AddKeyInChannel(TMovieSceneChannelHandle<ChannelType> ChannelHandle, TWeakObjectPtr<UMovieSceneSequence> Sequence, 
 		TWeakObjectPtr<UMovieSceneSection> Section, const FFrameNumber InTime, ScriptingKeyValueType& NewValue, float SubFrame, ESequenceTimeUnit TimeUnit, EMovieSceneKeyInterpolation Interpolation)
 	{
 		ChannelType* Channel = ChannelHandle.Get();
@@ -92,7 +92,7 @@ struct TMovieSceneScriptingChannel
 		return nullptr;
 	}
 
-	void RemoveKeyFromChannel(TMovieSceneChannelHandle<ChannelType> ChannelHandle, UMovieSceneScriptingKey* Key)
+	static void RemoveKeyFromChannel(TMovieSceneChannelHandle<ChannelType> ChannelHandle, UMovieSceneScriptingKey* Key)
 	{
 		if (!Key)
 		{
@@ -118,7 +118,7 @@ struct TMovieSceneScriptingChannel
 		FFrame::KismetExecutionMessage(TEXT("Invalid ChannelHandle for MovieSceneScriptingChannel, failed to remove key."), ELogVerbosity::Error);
 	}
 
-	TArray<UMovieSceneScriptingKey*> GetKeysInChannel(TMovieSceneChannelHandle<ChannelType> ChannelHandle, TWeakObjectPtr<UMovieSceneSequence> Sequence, TWeakObjectPtr<UMovieSceneSection> Section) const
+	static TArray<UMovieSceneScriptingKey*> GetKeysInChannel(TMovieSceneChannelHandle<ChannelType> ChannelHandle, TWeakObjectPtr<UMovieSceneSequence> Sequence, TWeakObjectPtr<UMovieSceneSection> Section)
 	{
 		TArray<UMovieSceneScriptingKey*> OutScriptingKeys;
 		ChannelType* Channel = ChannelHandle.Get();
@@ -146,7 +146,7 @@ struct TMovieSceneScriptingChannel
 		return OutScriptingKeys;
 	}
 
-	TArray<ScriptingKeyValueType> EvaluateKeysInChannel(TMovieSceneChannelHandle<ChannelType> ChannelHandle, TWeakObjectPtr<UMovieSceneSequence> Sequence, FSequencerScriptingRange ScriptingRange, FFrameRate FrameRate) const
+	static TArray<ScriptingKeyValueType> EvaluateKeysInChannel(TMovieSceneChannelHandle<ChannelType> ChannelHandle, TWeakObjectPtr<UMovieSceneSequence> Sequence, FSequencerScriptingRange ScriptingRange, FFrameRate FrameRate)
 	{
 		TArray<ScriptingKeyValueType> OutValues;
 		ChannelType* Channel = ChannelHandle.Get();
@@ -179,7 +179,7 @@ struct TMovieSceneScriptingChannel
 		return OutValues;
 	}
 
-	FSequencerScriptingRange ComputeEffectiveRangeInChannel(TMovieSceneChannelHandle<ChannelType> ChannelHandle, TWeakObjectPtr<UMovieSceneSequence> Sequence) const
+	static FSequencerScriptingRange ComputeEffectiveRangeInChannel(TMovieSceneChannelHandle<ChannelType> ChannelHandle, TWeakObjectPtr<UMovieSceneSequence> Sequence)
 	{
 		FSequencerScriptingRange ScriptingRange;
 		ChannelType* Channel = ChannelHandle.Get();
@@ -195,7 +195,7 @@ struct TMovieSceneScriptingChannel
 		return ScriptingRange;
 	}
 
-	void SetDefaultInChannel(TMovieSceneChannelHandle<ChannelType> ChannelHandle, TWeakObjectPtr<UMovieSceneSequence> Sequence, TWeakObjectPtr<UMovieSceneSection> Section, ScriptingKeyValueType& InDefaultValue)
+	static void SetDefaultInChannel(TMovieSceneChannelHandle<ChannelType> ChannelHandle, TWeakObjectPtr<UMovieSceneSequence> Sequence, TWeakObjectPtr<UMovieSceneSection> Section, ScriptingKeyValueType& InDefaultValue)
 	{
 		ChannelType* Channel = ChannelHandle.Get();
 		if (Channel)
@@ -216,7 +216,7 @@ struct TMovieSceneScriptingChannel
 		FFrame::KismetExecutionMessage(TEXT("Invalid ChannelHandle for MovieSceneScriptingChannel, failed to set default value."), ELogVerbosity::Error);
 	}
 
-	void RemoveDefaultFromChannel(TMovieSceneChannelHandle<ChannelType> ChannelHandle)
+	static void RemoveDefaultFromChannel(TMovieSceneChannelHandle<ChannelType> ChannelHandle)
 	{
 		ChannelType* Channel = ChannelHandle.Get();
 		if (Channel)
@@ -228,7 +228,7 @@ struct TMovieSceneScriptingChannel
 		FFrame::KismetExecutionMessage(TEXT("Invalid ChannelHandle for MovieSceneScriptingChannel, failed to remove default value."), ELogVerbosity::Error);
 	}
 
-	TOptional<ScriptingKeyValueType> GetDefaultFromChannel(TMovieSceneChannelHandle<ChannelType> ChannelHandle) const
+	static TOptional<ScriptingKeyValueType> GetDefaultFromChannel(TMovieSceneChannelHandle<ChannelType> ChannelHandle)
 	{
 		ChannelType* Channel = ChannelHandle.Get();
 		if (Channel)
diff --git a/Engine/Plugins/MovieScene/SequencerScripting/Source/SequencerScripting/Private/KeysAndChannels/MovieSceneScriptingDouble.h b/Engine/Plugins/MovieScene/SequencerScripting/Source/SequencerScripting/Private/KeysAndChannels/MovieSceneScriptingDouble.h
--- a/Engine/Plugins/MovieScene/SequencerScripting/Source/SequencerScripting/Private/KeysAndChannels/MovieSceneScriptingDouble.h
+++ b/Engine/Plugins/MovieScene/SequencerScripting/Source/SequencerScripting/Private/KeysAndChannels/MovieSceneScriptingDouble.h
@@ -223,9 +223,12 @@ public:
 };
 
 UCLASS(BlueprintType)
-class UMovieSceneScriptingDoubleChannel : public UMovieSceneScriptingChannel, public TMovieSceneScriptingChannel<FMovieSceneDoubleChannel, UMovieSceneScriptingDoubleKey, double>
+class UMovieSceneScriptingDoubleChannel : public UMovieSceneScriptingChannel
 {
 	GENERATED_BODY()
+
+	using Impl = TMovieSceneScriptingChannel<FMovieSceneDoubleChannel, UMovieSceneScriptingDoubleKey, double>;
+
 public:
 	/**
 	* Add a key to this channel. This initializes a new key and returns a reference to it.
@@ -239,7 +242,7 @@ public:
 	UFUNCTION(BlueprintCallable, Category = "Sequencer|Keys", meta = (DisplayName = "Add Key (Double)"))
 	UMovieSceneScriptingDoubleKey* AddKey(const FFrameNumber& InTime, double NewValue, float SubFrame = 0.f, ESequenceTimeUnit TimeUnit = ESequenceTimeUnit::DisplayRate, EMovieSceneKeyInterpolation InInterpolation = EMovieSceneKeyInterpolation::Auto)
 	{
-		return AddKeyInChannel(ChannelHandle, OwningSequence, OwningSection, InTime, NewValue, SubFrame, TimeUnit, InInterpolation);
+		return Impl::AddKeyInChannel(ChannelHandle, OwningSequence, OwningSection, InTime, NewValue, SubFrame, TimeUnit, InInterpolation);
 	}
 
 	/**
@@ -248,7 +251,7 @@ public:
 	UFUNCTION(BlueprintCallable, Category = "Sequencer|Keys", meta = (DisplayName = "Remove Key (Double)"))
 	virtual void RemoveKey(UMovieSceneScriptingKey* Key)
 	{
-		RemoveKeyFromChannel(ChannelHandle, Key);
+		Impl::RemoveKeyFromChannel(ChannelHandle, Key);
 	}
 
 	/**
@@ -259,7 +262,7 @@ public:
 	UFUNCTION(BlueprintCallable, Category = "Sequencer|Keys", meta = (DisplayName = "Get Keys (Double)"))
 	virtual TArray<UMovieSceneScriptingKey*> GetKeys() const override
 	{
-		return GetKeysInChannel(ChannelHandle, OwningSequence, OwningSection);
+		return Impl::GetKeysInChannel(ChannelHandle, OwningSequence, OwningSection);
 	}
 
 	/**
@@ -279,7 +282,7 @@ public:
 	UFUNCTION(BlueprintCallable, Category = "Sequencer|Keys", meta = (DisplayName = "Evaluate Keys (Double)"))
 	TArray<double> EvaluateKeys(FSequencerScriptingRange Range, FFrameRate FrameRate) const
 	{
-		return EvaluateKeysInChannel(ChannelHandle, OwningSequence, Range, FrameRate);
+		return Impl::EvaluateKeysInChannel(ChannelHandle, OwningSequence, Range, FrameRate);
 	}
 
 	/**
@@ -290,7 +293,7 @@ public:
 	UFUNCTION(BlueprintCallable, Category = "Sequencer|Keys", meta = (DisplayName = "Compute Effective Range (Double)"))
 	FSequencerScriptingRange ComputeEffectiveRange() const
 	{
-		return ComputeEffectiveRangeInChannel(ChannelHandle, OwningSequence);
+		return Impl::ComputeEffectiveRangeInChannel(ChannelHandle, OwningSequence);
 	}
 
 	/**
@@ -366,7 +369,7 @@ public:
 	UFUNCTION(BlueprintCallable, Category = "Sequencer|Keys", meta = (DisplayName = "Set Default (Double)"))
 	void SetDefault(double InDefaultValue)
 	{
-		SetDefaultInChannel(ChannelHandle, OwningSequence, OwningSection, InDefaultValue);
+		Impl::SetDefaultInChannel(ChannelHandle, OwningSequence, OwningSection, InDefaultValue);
 	}
 
 	/**
@@ -376,7 +379,7 @@ public:
 	UFUNCTION(BlueprintCallable, Category = "Sequencer|Keys", meta = (DisplayName = "Get Default (Double)"))
 	double GetDefault() const
 	{
-		TOptional<double> DefaultValue = GetDefaultFromChannel(ChannelHandle);
+		TOptional<double> DefaultValue = Impl::GetDefaultFromChannel(ChannelHandle);
 		return DefaultValue.IsSet() ? DefaultValue.GetValue() : 0.f;
 	}
 
@@ -386,7 +389,7 @@ public:
 	UFUNCTION(BlueprintCallable, Category = "Sequencer|Keys", meta = (DisplayName = "Remove Default (Double)"))
 	void RemoveDefault()
 	{
-		RemoveDefaultFromChannel(ChannelHandle);
+		Impl::RemoveDefaultFromChannel(ChannelHandle);
 	}
 
 	/**
@@ -395,7 +398,7 @@ public:
 	UFUNCTION(BlueprintCallable, Category = "Sequencer|Keys", meta = (DisplayName = "Has Default (Double)"))
 	bool HasDefault() const
 	{
-		return GetDefaultFromChannel(ChannelHandle).IsSet();
+		return Impl::GetDefaultFromChannel(ChannelHandle).IsSet();
 	}
 public:
 	TWeakObjectPtr<UMovieSceneSequence> OwningSequence;
diff --git a/Engine/Plugins/MovieScene/SequencerScripting/Source/SequencerScripting/Private/KeysAndChannels/MovieSceneScriptingEvent.h b/Engine/Plugins/MovieScene/SequencerScripting/Source/SequencerScripting/Private/KeysAndChannels/MovieSceneScriptingEvent.h
--- a/Engine/Plugins/MovieScene/SequencerScripting/Source/SequencerScripting/Private/KeysAndChannels/MovieSceneScriptingEvent.h
+++ b/Engine/Plugins/MovieScene/SequencerScripting/Source/SequencerScripting/Private/KeysAndChannels/MovieSceneScriptingEvent.h
@@ -63,9 +63,12 @@ public:
 };
 
 UCLASS(BlueprintType)
-class UMovieSceneScriptingEventChannel : public UMovieSceneScriptingChannel, public TMovieSceneScriptingChannel<FMovieSceneEventChannel, UMovieSceneScriptingEventKey, FMovieSceneEvent>
+class UMovieSceneScriptingEventChannel : public UMovieSceneScriptingChannel
 {
 	GENERATED_BODY()
+
+	using Impl = TMovieSceneScriptingChannel<FMovieSceneEventChannel, UMovieSceneScriptingEventKey, FMovieSceneEvent>;
+
 public:
 	/**
 	* Add a key to this channel. This initializes a new key and returns a reference to it.
@@ -78,7 +81,7 @@ public:
 	UFUNCTION(BlueprintCallable, Category = "Movie Scene | Keys", meta = (DisplayName = "Add Key (Event)"))
 	UMovieSceneScriptingEventKey* AddKey(const FFrameNumber& InTime, FMovieSceneEvent NewValue, float SubFrame = 0.f, ESequenceTimeUnit TimeUnit = ESequenceTimeUnit::DisplayRate)
 	{
-		return AddKeyInChannel(ChannelHandle, OwningSequence, OwningSection,InTime, NewValue, SubFrame, TimeUnit, EMovieSceneKeyInterpolation::Auto);
+		return Impl::AddKeyInChannel(ChannelHandle, OwningSequence, OwningSection,InTime, NewValue, SubFrame, TimeUnit, EMovieSceneKeyInterpolation::Auto);
 	}
 
 	/**
@@ -87,7 +90,7 @@ public:
 	UFUNCTION(BlueprintCallable, Category = "Movie Scene | Keys", meta = (DisplayName = "Remove Key (Event)"))
 	virtual void RemoveKey(UMovieSceneScriptingKey* Key)
 	{
-		RemoveKeyFromChannel(ChannelHandle, Key);
+		Impl::RemoveKeyFromChannel(ChannelHandle, Key);
 	}
 
 	/**
@@ -98,7 +101,7 @@ public:
 	UFUNCTION(BlueprintCallable, Category = "Movie Scene | Keys", meta = (DisplayName = "Get Keys (Event)"))
 	virtual TArray<UMovieSceneScriptingKey*> GetKeys() const override
 	{
-		return GetKeysInChannel(ChannelHandle, OwningSequence, OwningSection);
+		return Impl::GetKeysInChannel(ChannelHandle, OwningSequence, OwningSection);
 	}
 
 public:
@@ -106,4 +109,4 @@ public:
 	TMovieSceneChannelHandle<FMovieSceneEventChannel> ChannelHandle;
 	TWeakObjectPtr<UMovieSceneSection> OwningSection;
 
-};
\ No newline at end of file
+};
diff --git a/Engine/Plugins/MovieScene/SequencerScripting/Source/SequencerScripting/Private/KeysAndChannels/MovieSceneScriptingFloat.h b/Engine/Plugins/MovieScene/SequencerScripting/Source/SequencerScripting/Private/KeysAndChannels/MovieSceneScriptingFloat.h
--- a/Engine/Plugins/MovieScene/SequencerScripting/Source/SequencerScripting/Private/KeysAndChannels/MovieSceneScriptingFloat.h
+++ b/Engine/Plugins/MovieScene/SequencerScripting/Source/SequencerScripting/Private/KeysAndChannels/MovieSceneScriptingFloat.h
@@ -7,6 +7,7 @@
 #include "MovieSceneScriptingChannel.h"
 #include "Channels/MovieSceneChannelTraits.h"
 #include "Channels/MovieSceneFloatChannel.h"
+#include "Channels/MovieSceneDoubleChannel.h"
 #include "Channels/MovieSceneChannelHandle.h"
 #include "KeysAndChannels/MovieSceneScriptingChannel.h"
 #include "KeyParams.h"
@@ -19,7 +20,7 @@
 * Stores a reference to the data so changes to this class are forwarded onto the underlying data structures.
 */
 UCLASS(BlueprintType)
-class UMovieSceneScriptingFloatKey : public UMovieSceneScriptingKey, public TMovieSceneScriptingKey<FMovieSceneFloatChannel, FMovieSceneFloatValue>
+class UMovieSceneScriptingFloatKey : public UMovieSceneScriptingKey
 {
 	GENERATED_BODY()
 public:
@@ -29,7 +30,7 @@ public:
 	* @return			The time of this key which combines both the frame number and the sub-frame it is on. Sub-frame will be zero if you request Tick Resolution.
 	*/
 	UFUNCTION(BlueprintCallable, Category = "Sequencer|Keys", meta = (DisplayName = "Get Time (Float)"))
-	virtual FFrameTime GetTime(ESequenceTimeUnit TimeUnit = ESequenceTimeUnit::DisplayRate) const override { return GetTimeFromChannel(KeyHandle, OwningSequence, TimeUnit); }
+	virtual FFrameTime GetTime(ESequenceTimeUnit TimeUnit = ESequenceTimeUnit::DisplayRate) const override PURE_VIRTUAL(UMovieSceneScriptingFloatKey::GetTime, return FFrameTime(););
 	
 	/**
 	* Sets the time for this key in the owning channel. Will replace any key that already exists at that frame number in this channel.
@@ -38,194 +39,329 @@ public:
 	* @param TimeUnit		Should the NewFrameNumber be interpreted as Display Rate frames or in Tick Resolution?
 	*/
 	UFUNCTION(BlueprintCallable, Category = "Sequencer|Keys", meta = (DisplayName = "Set Time (Float)"))
-	void SetTime(const FFrameNumber& NewFrameNumber, float SubFrame = 0.f, ESequenceTimeUnit TimeUnit = ESequenceTimeUnit::DisplayRate) { SetTimeInChannel(KeyHandle, OwningSequence, NewFrameNumber, TimeUnit, SubFrame); }
+	virtual void SetTime(const FFrameNumber& NewFrameNumber, float SubFrame = 0.f, ESequenceTimeUnit TimeUnit = ESequenceTimeUnit::DisplayRate) PURE_VIRTUAL(UMovieSceneScriptingFloatKey::SetTime);
 
 	/**
 	* Gets the value for this key from the owning channel.
 	* @return	The float value this key represents.
 	*/
 	UFUNCTION(BlueprintCallable, Category = "Sequencer|Keys", meta = (DisplayName = "Get Value (Float)"))
-	float GetValue() const
-	{
-		return GetValueFromChannel(KeyHandle).Value;
-	}
+	virtual float GetValue() const PURE_VIRTUAL(UMovieSceneScriptingFloatKey::GetValue, return 0;);
 
 	/**
 	* Sets the value for this key, reflecting it in the owning channel.
 	* @param InNewValue	The new float value for this key.
 	*/
 	UFUNCTION(BlueprintCallable, Category = "Sequencer|Keys", meta = (DisplayName = "Set Value (Float)"))
-	void SetValue(float InNewValue)
-	{
-		FMovieSceneFloatValue ExistValue = GetValueFromChannel(KeyHandle);
-		ExistValue.Value = InNewValue;
-		SetValueInChannel(KeyHandle, ExistValue);
-	}
+	virtual void SetValue(float InNewValue) PURE_VIRTUAL(UMovieSceneScriptingFloatKey::SetValue);
 
 	/**
 	* Gets the interpolation mode for this key from the owning channel.
 	* @return	Interpolation mode this key uses to interpolate between this key and the next.
 	*/
 	UFUNCTION(BlueprintCallable, Category = "Sequencer|Keys")
-	ERichCurveInterpMode GetInterpolationMode() const
-	{
-		return GetValueFromChannel(KeyHandle).InterpMode;
-	}
+	virtual ERichCurveInterpMode GetInterpolationMode() const PURE_VIRTUAL(UMovieSceneScriptingFloatKey::GetInterpolationMode, return ERichCurveInterpMode::RCIM_None;);
 
 	/**
 	* Sets the interpolation mode for this key, reflecting it in the owning channel.
 	* @param InNewValue	Interpolation mode this key should use to interpolate between this key and the next.
 	*/
 	UFUNCTION(BlueprintCallable, Category = "Sequencer|Keys")
-	void SetInterpolationMode(ERichCurveInterpMode InNewValue)
-	{
-		FMovieSceneFloatValue ExistValue = GetValueFromChannel(KeyHandle);
-		ExistValue.InterpMode = InNewValue;
-		SetValueInChannel(KeyHandle, ExistValue);
-	}
+	virtual void SetInterpolationMode(ERichCurveInterpMode InNewValue) PURE_VIRTUAL(UMovieSceneScriptingFloatKey::SetInterpolationMode);
 
 	/**
 	* Gets the tangent mode for this key from the owning channel.
 	* @return	Tangent mode that this key is using specifying which tangent values are respected when evaluating.
 	*/
 	UFUNCTION(BlueprintCallable, Category = "Sequencer|Keys")
-	ERichCurveTangentMode GetTangentMode() const
-	{
-		return GetValueFromChannel(KeyHandle).TangentMode;
-	}
+	virtual ERichCurveTangentMode GetTangentMode() const PURE_VIRTUAL(UMovieSceneScriptingFloatKey::GetTangentMode, return ERichCurveTangentMode::RCTM_None;);
 
 	/**
 	* Sets the tangent mode for this key, reflecting it in the owning channel.
 	* @param InNewValue	Tangent mode that this key should use to specify which tangent values are respected when evaluating. See ERichCurveTangentMode for more details.
 	*/
 	UFUNCTION(BlueprintCallable, Category = "Sequencer|Keys")
-	void SetTangentMode(ERichCurveTangentMode InNewValue)
-	{
-		FMovieSceneFloatValue ExistValue = GetValueFromChannel(KeyHandle);
-		ExistValue.TangentMode = InNewValue;
-		SetValueInChannel(KeyHandle, ExistValue);
-	}
+	virtual void SetTangentMode(ERichCurveTangentMode InNewValue) PURE_VIRTUAL(UMovieSceneScriptingFloatKey::SetTangentMode);
 
 	/**
 	* If Interpolation Mode is RCIM_Cubic, the arriving tangent at this key
 	* @return Arrival Tangent value. Represents the geometric tangents in the form of "tan(y/x)" where y is the key's value and x is the seconds (both relative to key)
 	*/
 	UFUNCTION(BlueprintCallable, Category = "Sequencer|Keys")
-	float GetArriveTangent() const
-	{
-		return GetValueFromChannel(KeyHandle).Tangent.ArriveTangent;
-	}
+	virtual float GetArriveTangent() const PURE_VIRTUAL(UMovieSceneScriptingFloatKey::GetArriveTangent, return 0;);
 
 	/**
 	* If Interpolation Mode is RCIM_Cubic, the arriving tangent at this key.
 	* @param InNewValue	Represents the geometric tangents in the form of "tan(y/x)" where y is the key's value and x is the seconds (both relative to key)
 	*/
 	UFUNCTION(BlueprintCallable, Category = "Sequencer|Keys")
-	void SetArriveTangent(float InNewValue)
-	{
-		FMovieSceneFloatValue ExistKeyValue = GetValueFromChannel(KeyHandle);
-		FMovieSceneTangentData& ExistTangentData = ExistKeyValue.Tangent;
-		ExistTangentData.ArriveTangent = InNewValue;
-		SetValueInChannel(KeyHandle, ExistKeyValue);
-	}
+	virtual void SetArriveTangent(float InNewValue) PURE_VIRTUAL(UMovieSceneScriptingFloatKey::SetArriveTangent);
 
 	/**
 	* If Interpolation Mode is RCIM_Cubic, the leaving tangent at this key
 	* @return Leaving Tangent value. Represents the geometric tangents in the form of "tan(y/x)" where y is the key's value and x is the seconds (both relative to key)
 	*/
 	UFUNCTION(BlueprintCallable, Category = "Sequencer|Keys")
-	float GetLeaveTangent() const
-	{
-		return GetValueFromChannel(KeyHandle).Tangent.LeaveTangent;
-	}
+	virtual float GetLeaveTangent() const PURE_VIRTUAL(UMovieSceneScriptingFloatKey::GetLeaveTangent, return 0;);
 
 	/**
 	* If Interpolation Mode is RCIM_Cubic, the leaving tangent at this key.
 	* @param InNewValue	Represents the geometric tangents in the form of "tan(y/x)" where y is the key's value and x is the seconds (both relative to key)
 	*/
 	UFUNCTION(BlueprintCallable, Category = "Sequencer|Keys")
-	void SetLeaveTangent(float InNewValue)
-	{
-		FMovieSceneFloatValue ExistKeyValue = GetValueFromChannel(KeyHandle);
-		FMovieSceneTangentData& ExistTangentData = ExistKeyValue.Tangent;
-		ExistTangentData.LeaveTangent = InNewValue;
-		SetValueInChannel(KeyHandle, ExistKeyValue);
-	}
+	virtual void SetLeaveTangent(float InNewValue) PURE_VIRTUAL(UMovieSceneScriptingFloatKey::SetLeaveTangent);
 
 	/**
 	* If Interpolation Mode is RCIM_Cubic, the tangent weight mode at this key
 	* @return Tangent Weight Mode. See ERichCurveTangentWeightMode for more detail on what each mode does.
 	*/
 	UFUNCTION(BlueprintCallable, Category = "Sequencer|Keys")
-	ERichCurveTangentWeightMode GetTangentWeightMode() const
-	{
-		return GetValueFromChannel(KeyHandle).Tangent.TangentWeightMode;
-	}
+	virtual ERichCurveTangentWeightMode GetTangentWeightMode() const PURE_VIRTUAL(UMovieSceneScriptingFloatKey::GetTangentWeightMode, return ERichCurveTangentWeightMode::RCTWM_WeightedNone;);
 
 	/**
 	* If Interpolation Mode is RCIM_Cubic, the tangent weight mode at this key.
 	* @param InNewValue	Specifies which tangent weights should be respected when evaluating the key. 
 	*/
 	UFUNCTION(BlueprintCallable, Category = "Sequencer|Keys")
-	void SetTangentWeightMode(ERichCurveTangentWeightMode InNewValue)
-	{
-		FMovieSceneFloatValue ExistKeyValue = GetValueFromChannel(KeyHandle);
-		FMovieSceneTangentData& ExistTangentData = ExistKeyValue.Tangent;
-		ExistTangentData.TangentWeightMode = InNewValue;
-		SetValueInChannel(KeyHandle, ExistKeyValue);
-	}
+	virtual void SetTangentWeightMode(ERichCurveTangentWeightMode InNewValue) PURE_VIRTUAL(UMovieSceneScriptingFloatKey::SetTangentWeightMode);
 
 	/**
 	* If Tangent Weight Mode is RCTWM_WeightedArrive or RCTWM_WeightedBoth, the weight of the arriving tangent on the left side.
 	* @return Tangent Weight. Represents the length of the hypotenuse in the form of "sqrt(x*x+y*y)" using the same definitions for x and y as tangents.
 	*/
 	UFUNCTION(BlueprintCallable, Category = "Sequencer|Keys")
-	float GetArriveTangentWeight() const
-	{
-		return GetValueFromChannel(KeyHandle).Tangent.ArriveTangentWeight;
-	}
+	virtual float GetArriveTangentWeight() const PURE_VIRTUAL(UMovieSceneScriptingFloatKey::GetArriveTangentWeight, return 0;);
 
 	/**
 	* If Tangent Weight Mode is RCTWM_WeightedArrive or RCTWM_WeightedBoth, the weight of the arriving tangent on the left side.
 	* @param InNewValue	Specifies the new arriving tangent weight. Represents the length of the hypotenuse in the form of "sqrt(x*x+y*y)" using the same definitions for x and y as tangents.
 	*/
 	UFUNCTION(BlueprintCallable, Category = "Sequencer|Keys")
-	void SetArriveTangentWeight(float InNewValue)
-	{
-		FMovieSceneFloatValue ExistKeyValue = GetValueFromChannel(KeyHandle);
-		FMovieSceneTangentData& ExistTangentData = ExistKeyValue.Tangent;
-		ExistTangentData.ArriveTangentWeight = InNewValue;
-		SetValueInChannel(KeyHandle, ExistKeyValue);
-	}
+	virtual void SetArriveTangentWeight(float InNewValue) PURE_VIRTUAL(UMovieSceneScriptingFloatKey::SetArriveTangentWeight);
 
 	/**
 	* If Tangent Weight Mode is RCTWM_WeightedLeave or RCTWM_WeightedBoth, the weight of the leaving tangent on the right side.
 	* @return Tangent Weight. Represents the length of the hypotenuse in the form of "sqrt(x*x+y*y)" using the same definitions for x and y as tangents.
 	*/
 	UFUNCTION(BlueprintCallable, Category = "Sequencer|Keys")
-	float GetLeaveTangentWeight() const
-	{
-		return GetValueFromChannel(KeyHandle).Tangent.LeaveTangentWeight;
-	}
+	virtual float GetLeaveTangentWeight() const PURE_VIRTUAL(UMovieSceneScriptingFloatKey::GetLeaveTangentWeight, return 0;);
 
 	/**
 	* If Tangent Weight Mode is RCTWM_WeightedLeave or RCTWM_WeightedBoth, the weight of the leaving tangent on the right side.
 	* @param InNewValue	Specifies the new leaving tangent weight. Represents the length of the hypotenuse in the form of "sqrt(x*x+y*y)" using the same definitions for x and y as tangents.
 	*/
 	UFUNCTION(BlueprintCallable, Category = "Sequencer|Keys")
-	void SetLeaveTangentWeight(float InNewValue)
+	virtual void SetLeaveTangentWeight(float InNewValue) PURE_VIRTUAL(UMovieSceneScriptingFloatKey::SetLeaveTangentWeight);
+};
+
+UCLASS(BlueprintType)
+class UMovieSceneScriptingActualFloatKey : public UMovieSceneScriptingFloatKey, public TMovieSceneScriptingKey<FMovieSceneFloatChannel, FMovieSceneFloatValue>
+{
+	GENERATED_BODY()
+public:
+	virtual FFrameTime GetTime(ESequenceTimeUnit TimeUnit = ESequenceTimeUnit::DisplayRate) const override
+	{
+		return GetTimeFromChannel(KeyHandle, OwningSequence, TimeUnit);
+	}
+	virtual void SetTime(const FFrameNumber& NewFrameNumber, float SubFrame = 0.f, ESequenceTimeUnit TimeUnit = ESequenceTimeUnit::DisplayRate) override
+	{
+		SetTimeInChannel(KeyHandle, OwningSequence, NewFrameNumber, TimeUnit, SubFrame);
+	}
+	virtual float GetValue() const override
+	{
+		return GetValueFromChannel(KeyHandle).Value;
+	}
+	virtual void SetValue(float InNewValue) override
+	{
+		FMovieSceneFloatValue ExistValue = GetValueFromChannel(KeyHandle);
+		ExistValue.Value = InNewValue;
+		SetValueInChannel(KeyHandle, ExistValue);
+	}
+	virtual ERichCurveInterpMode GetInterpolationMode() const override
+	{
+		return GetValueFromChannel(KeyHandle).InterpMode;
+	}
+	virtual void SetInterpolationMode(ERichCurveInterpMode InNewValue) override
+	{
+		FMovieSceneFloatValue ExistValue = GetValueFromChannel(KeyHandle);
+		ExistValue.InterpMode = InNewValue;
+		SetValueInChannel(KeyHandle, ExistValue);
+	}
+	virtual ERichCurveTangentMode GetTangentMode() const override
+	{
+		return GetValueFromChannel(KeyHandle).TangentMode;
+	}
+	virtual void SetTangentMode(ERichCurveTangentMode InNewValue) override
+	{
+		FMovieSceneFloatValue ExistValue = GetValueFromChannel(KeyHandle);
+		ExistValue.TangentMode = InNewValue;
+		SetValueInChannel(KeyHandle, ExistValue);
+	}
+	virtual float GetArriveTangent() const override
+	{
+		return GetValueFromChannel(KeyHandle).Tangent.ArriveTangent;
+	}
+	virtual void SetArriveTangent(float InNewValue) override
 	{
 		FMovieSceneFloatValue ExistKeyValue = GetValueFromChannel(KeyHandle);
 		FMovieSceneTangentData& ExistTangentData = ExistKeyValue.Tangent;
+		ExistTangentData.ArriveTangent = InNewValue;
+		SetValueInChannel(KeyHandle, ExistKeyValue);
+	}
+	virtual float GetLeaveTangent() const override
+	{
+		return GetValueFromChannel(KeyHandle).Tangent.LeaveTangent;
+	}
+	virtual void SetLeaveTangent(float InNewValue) override
+	{
+		FMovieSceneFloatValue ExistKeyValue = GetValueFromChannel(KeyHandle);
+		FMovieSceneTangentData& ExistTangentData = ExistKeyValue.Tangent;
+		ExistTangentData.LeaveTangent = InNewValue;
+		SetValueInChannel(KeyHandle, ExistKeyValue);
+	}
+	virtual ERichCurveTangentWeightMode GetTangentWeightMode() const override
+	{
+		return GetValueFromChannel(KeyHandle).Tangent.TangentWeightMode;
+	}
+	virtual void SetTangentWeightMode(ERichCurveTangentWeightMode InNewValue) override
+	{
+		FMovieSceneFloatValue ExistKeyValue = GetValueFromChannel(KeyHandle);
+		FMovieSceneTangentData& ExistTangentData = ExistKeyValue.Tangent;
+		ExistTangentData.TangentWeightMode = InNewValue;
+		SetValueInChannel(KeyHandle, ExistKeyValue);
+	}
+	virtual float GetArriveTangentWeight() const override
+	{
+		return GetValueFromChannel(KeyHandle).Tangent.ArriveTangentWeight;
+	}
+	virtual void SetArriveTangentWeight(float InNewValue) override
+	{
+		FMovieSceneFloatValue ExistKeyValue = GetValueFromChannel(KeyHandle);
+		FMovieSceneTangentData& ExistTangentData = ExistKeyValue.Tangent;
+		ExistTangentData.ArriveTangentWeight = InNewValue;
+		SetValueInChannel(KeyHandle, ExistKeyValue);
+	}
+	virtual float GetLeaveTangentWeight() const override
+	{
+		return GetValueFromChannel(KeyHandle).Tangent.LeaveTangentWeight;
+	}
+	virtual void SetLeaveTangentWeight(float InNewValue) override
+	{
+		FMovieSceneFloatValue ExistKeyValue = GetValueFromChannel(KeyHandle);
+		FMovieSceneTangentData& ExistTangentData = ExistKeyValue.Tangent;
+		ExistTangentData.LeaveTangentWeight = InNewValue;
+		SetValueInChannel(KeyHandle, ExistKeyValue);
+	}
+};
+
+UCLASS(BlueprintType)
+class UMovieSceneScriptingDoubleAsFloatKey : public UMovieSceneScriptingFloatKey, public TMovieSceneScriptingKey<FMovieSceneDoubleChannel, FMovieSceneDoubleValue>
+{
+	GENERATED_BODY()
+public:
+	virtual FFrameTime GetTime(ESequenceTimeUnit TimeUnit = ESequenceTimeUnit::DisplayRate) const override
+	{
+		return GetTimeFromChannel(KeyHandle, OwningSequence, TimeUnit);
+	}
+	virtual void SetTime(const FFrameNumber& NewFrameNumber, float SubFrame = 0.f, ESequenceTimeUnit TimeUnit = ESequenceTimeUnit::DisplayRate) override
+	{
+		SetTimeInChannel(KeyHandle, OwningSequence, NewFrameNumber, TimeUnit, SubFrame);
+	}
+	virtual float GetValue() const override
+	{
+		return (float)GetValueFromChannel(KeyHandle).Value;
+	}
+	virtual void SetValue(float InNewValue) override
+	{
+		FMovieSceneDoubleValue ExistValue = GetValueFromChannel(KeyHandle);
+		ExistValue.Value = (double)InNewValue;
+		SetValueInChannel(KeyHandle, ExistValue);
+	}
+	virtual ERichCurveInterpMode GetInterpolationMode() const override
+	{
+		return GetValueFromChannel(KeyHandle).InterpMode;
+	}
+	virtual void SetInterpolationMode(ERichCurveInterpMode InNewValue) override
+	{
+		FMovieSceneDoubleValue ExistValue = GetValueFromChannel(KeyHandle);
+		ExistValue.InterpMode = InNewValue;
+		SetValueInChannel(KeyHandle, ExistValue);
+	}
+	virtual ERichCurveTangentMode GetTangentMode() const override
+	{
+		return GetValueFromChannel(KeyHandle).TangentMode;
+	}
+	virtual void SetTangentMode(ERichCurveTangentMode InNewValue) override
+	{
+		FMovieSceneDoubleValue ExistValue = GetValueFromChannel(KeyHandle);
+		ExistValue.TangentMode = InNewValue;
+		SetValueInChannel(KeyHandle, ExistValue);
+	}
+	virtual float GetArriveTangent() const override
+	{
+		return GetValueFromChannel(KeyHandle).Tangent.ArriveTangent;
+	}
+	virtual void SetArriveTangent(float InNewValue) override
+	{
+		FMovieSceneDoubleValue ExistKeyValue = GetValueFromChannel(KeyHandle);
+		FMovieSceneTangentData& ExistTangentData = ExistKeyValue.Tangent;
+		ExistTangentData.ArriveTangent = InNewValue;
+		SetValueInChannel(KeyHandle, ExistKeyValue);
+	}
+	virtual float GetLeaveTangent() const override
+	{
+		return GetValueFromChannel(KeyHandle).Tangent.LeaveTangent;
+	}
+	virtual void SetLeaveTangent(float InNewValue) override
+	{
+		FMovieSceneDoubleValue ExistKeyValue = GetValueFromChannel(KeyHandle);
+		FMovieSceneTangentData& ExistTangentData = ExistKeyValue.Tangent;
+		ExistTangentData.LeaveTangent = InNewValue;
+		SetValueInChannel(KeyHandle, ExistKeyValue);
+	}
+	virtual ERichCurveTangentWeightMode GetTangentWeightMode() const override
+	{
+		return GetValueFromChannel(KeyHandle).Tangent.TangentWeightMode;
+	}
+	virtual void SetTangentWeightMode(ERichCurveTangentWeightMode InNewValue) override
+	{
+		FMovieSceneDoubleValue ExistKeyValue = GetValueFromChannel(KeyHandle);
+		FMovieSceneTangentData& ExistTangentData = ExistKeyValue.Tangent;
+		ExistTangentData.TangentWeightMode = InNewValue;
+		SetValueInChannel(KeyHandle, ExistKeyValue);
+	}
+	virtual float GetArriveTangentWeight() const override
+	{
+		return GetValueFromChannel(KeyHandle).Tangent.ArriveTangentWeight;
+	}
+	virtual void SetArriveTangentWeight(float InNewValue) override
+	{
+		FMovieSceneDoubleValue ExistKeyValue = GetValueFromChannel(KeyHandle);
+		FMovieSceneTangentData& ExistTangentData = ExistKeyValue.Tangent;
+		ExistTangentData.ArriveTangentWeight = InNewValue;
+		SetValueInChannel(KeyHandle, ExistKeyValue);
+	}
+	virtual float GetLeaveTangentWeight() const override
+	{
+		return GetValueFromChannel(KeyHandle).Tangent.LeaveTangentWeight;
+	}
+	virtual void SetLeaveTangentWeight(float InNewValue) override
+	{
+		FMovieSceneDoubleValue ExistKeyValue = GetValueFromChannel(KeyHandle);
+		FMovieSceneTangentData& ExistTangentData = ExistKeyValue.Tangent;
 		ExistTangentData.LeaveTangentWeight = InNewValue;
 		SetValueInChannel(KeyHandle, ExistKeyValue);
 	}
 };
 
 UCLASS(BlueprintType)
-class UMovieSceneScriptingFloatChannel : public UMovieSceneScriptingChannel, public TMovieSceneScriptingChannel<FMovieSceneFloatChannel, UMovieSceneScriptingFloatKey, float>
+class UMovieSceneScriptingFloatChannel : public UMovieSceneScriptingChannel
 {
 	GENERATED_BODY()
+
+	using FloatImpl = TMovieSceneScriptingChannel<FMovieSceneFloatChannel, UMovieSceneScriptingActualFloatKey, float>;
+	using DoubleImpl = TMovieSceneScriptingChannel<FMovieSceneDoubleChannel, UMovieSceneScriptingDoubleAsFloatKey, double>;
+
 public:
 	/**
 	* Add a key to this channel. This initializes a new key and returns a reference to it.
@@ -239,7 +375,15 @@ public:
 	UFUNCTION(BlueprintCallable, Category = "Sequencer|Keys", meta = (DisplayName = "Add Key (Float)"))
 	UMovieSceneScriptingFloatKey* AddKey(const FFrameNumber& InTime, float NewValue, float SubFrame = 0.f, ESequenceTimeUnit TimeUnit = ESequenceTimeUnit::DisplayRate, EMovieSceneKeyInterpolation InInterpolation = EMovieSceneKeyInterpolation::Auto)
 	{
-		return AddKeyInChannel(ChannelHandle, OwningSequence, OwningSection, InTime, NewValue, SubFrame, TimeUnit, InInterpolation);
+		if (FloatChannelHandle.Get())
+		{
+			return FloatImpl::AddKeyInChannel(FloatChannelHandle, OwningSequence, OwningSection, InTime, NewValue, SubFrame, TimeUnit, InInterpolation);
+		}
+		else
+		{
+			double DoubleNewValue(NewValue);
+			return DoubleImpl::AddKeyInChannel(DoubleChannelHandle, OwningSequence, OwningSection, InTime, DoubleNewValue, SubFrame, TimeUnit, InInterpolation);
+		}
 	}
 
 	/**
@@ -248,7 +392,14 @@ public:
 	UFUNCTION(BlueprintCallable, Category = "Sequencer|Keys", meta = (DisplayName = "Remove Key (Float)"))
 	virtual void RemoveKey(UMovieSceneScriptingKey* Key)
 	{
-		RemoveKeyFromChannel(ChannelHandle, Key);
+		if (FloatChannelHandle.Get())
+		{
+			FloatImpl::RemoveKeyFromChannel(FloatChannelHandle, Key);
+		}
+		else
+		{
+			DoubleImpl::RemoveKeyFromChannel(DoubleChannelHandle, Key);
+		}
 	}
 
 	/**
@@ -259,7 +410,14 @@ public:
 	UFUNCTION(BlueprintCallable, Category = "Sequencer|Keys", meta = (DisplayName = "Get Keys (Float)"))
 	virtual TArray<UMovieSceneScriptingKey*> GetKeys() const override
 	{
-		return GetKeysInChannel(ChannelHandle, OwningSequence, OwningSection);
+		if (FloatChannelHandle.Get())
+		{
+			return FloatImpl::GetKeysInChannel(FloatChannelHandle, OwningSequence, OwningSection);
+		}
+		else
+		{
+			return DoubleImpl::GetKeysInChannel(DoubleChannelHandle, OwningSequence, OwningSection);
+		}
 	}
 
 	/**
@@ -268,7 +426,18 @@ public:
 	UFUNCTION(BlueprintCallable, Category = "Sequencer|Keys", meta = (DisplayName = "Get Num Keys (Float)"))
 	int32 GetNumKeys() const
 	{
-		return ChannelHandle.Get() ? ChannelHandle.Get()->GetNumKeys() : 0;
+		if (FloatChannelHandle.Get())
+		{
+			return FloatChannelHandle.Get()->GetNumKeys();
+		}
+		else if (DoubleChannelHandle.Get())
+		{
+			return DoubleChannelHandle.Get()->GetNumKeys();
+		}
+		else
+		{
+			return 0;
+		}
 	}
 
 	/**
@@ -279,7 +448,18 @@ public:
 	UFUNCTION(BlueprintCallable, Category = "Sequencer|Keys", meta = (DisplayName = "Evaluate Keys (Float)"))
 	TArray<float> EvaluateKeys(FSequencerScriptingRange Range, FFrameRate FrameRate) const
 	{
-		return EvaluateKeysInChannel(ChannelHandle, OwningSequence, Range, FrameRate);
+		if (FloatChannelHandle.Get())
+		{
+			return FloatImpl::EvaluateKeysInChannel(FloatChannelHandle, OwningSequence, Range, FrameRate);
+		}
+		else
+		{
+			TArray<double> DoubleValues = DoubleImpl::EvaluateKeysInChannel(DoubleChannelHandle, OwningSequence, Range, FrameRate);
+
+			TArray<float> FloatValues;
+			Algo::Transform(DoubleValues, FloatValues, [](double Value) { return (float)Value; });
+			return FloatValues;
+		}
 	}
 
 	/**
@@ -290,7 +470,14 @@ public:
 	UFUNCTION(BlueprintCallable, Category = "Sequencer|Keys", meta = (DisplayName = "Compute Effective Range (Float)"))
 	FSequencerScriptingRange ComputeEffectiveRange() const
 	{
-		return ComputeEffectiveRangeInChannel(ChannelHandle, OwningSequence);
+		if (FloatChannelHandle.Get())
+		{
+			return FloatImpl::ComputeEffectiveRangeInChannel(FloatChannelHandle, OwningSequence);
+		}
+		else
+		{
+			return DoubleImpl::ComputeEffectiveRangeInChannel(DoubleChannelHandle, OwningSequence);
+		}
 	}
 
 	/**
@@ -299,10 +486,16 @@ public:
 	UFUNCTION(BlueprintCallable, Category = "Sequencer|Keys")
 	ERichCurveExtrapolation GetPreInfinityExtrapolation() const
 	{
-		FMovieSceneFloatChannel* Channel = ChannelHandle.Get();
-		if (Channel)
+		FMovieSceneFloatChannel* FloatChannel = FloatChannelHandle.Get();
+		if (FloatChannel)
+		{
+			return FloatChannel->PreInfinityExtrap;
+		}
+
+		FMovieSceneDoubleChannel* DoubleChannel = DoubleChannelHandle.Get();
+		if (DoubleChannel)
 		{
-			return Channel->PreInfinityExtrap;
+			return DoubleChannel->PreInfinityExtrap;
 		}
 
 		FFrame::KismetExecutionMessage(TEXT("Invalid ChannelHandle for MovieSceneScriptingChannel, failed to get pre-infinity extrapolation."), ELogVerbosity::Error);
@@ -316,10 +509,17 @@ public:
 	UFUNCTION(BlueprintCallable, Category = "Sequencer|Keys")
 	void SetPreInfinityExtrapolation(ERichCurveExtrapolation InExtrapolation)
 	{
-		FMovieSceneFloatChannel* Channel = ChannelHandle.Get();
-		if (Channel)
+		FMovieSceneFloatChannel* FloatChannel = FloatChannelHandle.Get();
+		if (FloatChannel)
+		{
+			FloatChannel->PreInfinityExtrap = InExtrapolation;
+			return;
+		}
+
+		FMovieSceneDoubleChannel* DoubleChannel = DoubleChannelHandle.Get();
+		if (DoubleChannel)
 		{
-			Channel->PreInfinityExtrap = InExtrapolation;
+			DoubleChannel->PreInfinityExtrap = InExtrapolation;
 			return;
 		}
 
@@ -332,10 +532,16 @@ public:
 	UFUNCTION(BlueprintCallable, Category = "Sequencer|Keys")
 	ERichCurveExtrapolation GetPostInfinityExtrapolation() const
 	{
-		FMovieSceneFloatChannel* Channel = ChannelHandle.Get();
-		if (Channel)
+		FMovieSceneFloatChannel* FloatChannel = FloatChannelHandle.Get();
+		if (FloatChannel)
 		{
-			return Channel->PostInfinityExtrap;
+			return FloatChannel->PostInfinityExtrap;
+		}
+
+		FMovieSceneDoubleChannel* DoubleChannel = DoubleChannelHandle.Get();
+		if (DoubleChannel)
+		{
+			return DoubleChannel->PostInfinityExtrap;
 		}
 
 		FFrame::KismetExecutionMessage(TEXT("Invalid ChannelHandle for MovieSceneScriptingChannel, failed to get post-infinity extrapolation."), ELogVerbosity::Error);
@@ -349,10 +555,17 @@ public:
 	UFUNCTION(BlueprintCallable, Category = "Sequencer|Keys")
 	void SetPostInfinityExtrapolation(ERichCurveExtrapolation InExtrapolation)
 	{
-		FMovieSceneFloatChannel* Channel = ChannelHandle.Get();
-		if (Channel)
+		FMovieSceneFloatChannel* FloatChannel = FloatChannelHandle.Get();
+		if (FloatChannel)
+		{
+			FloatChannel->PostInfinityExtrap = InExtrapolation;
+			return;
+		}
+
+		FMovieSceneDoubleChannel* DoubleChannel = DoubleChannelHandle.Get();
+		if (DoubleChannel)
 		{
-			Channel->PostInfinityExtrap = InExtrapolation;
+			DoubleChannel->PostInfinityExtrap = InExtrapolation;
 			return;
 		}
 
@@ -366,7 +579,15 @@ public:
 	UFUNCTION(BlueprintCallable, Category = "Sequencer|Keys", meta = (DisplayName = "Set Default (Float)"))
 	void SetDefault(float InDefaultValue)
 	{
-		SetDefaultInChannel(ChannelHandle, OwningSequence, OwningSection, InDefaultValue);
+		if (FloatChannelHandle.Get())
+		{
+			FloatImpl::SetDefaultInChannel(FloatChannelHandle, OwningSequence, OwningSection, InDefaultValue);
+		}
+		else
+		{
+			double DoubleDefaultValue(InDefaultValue);
+			DoubleImpl::SetDefaultInChannel(DoubleChannelHandle, OwningSequence, OwningSection, DoubleDefaultValue);
+		}
 	}
 
 	/**
@@ -376,8 +597,16 @@ public:
 	UFUNCTION(BlueprintCallable, Category = "Sequencer|Keys", meta = (DisplayName = "Get Default (Float)"))
 	float GetDefault() const
 	{
-		TOptional<float> DefaultValue = GetDefaultFromChannel(ChannelHandle);
-		return DefaultValue.IsSet() ? DefaultValue.GetValue() : 0.f;
+		if (FloatChannelHandle.Get())
+		{
+			TOptional<float> DefaultValue = FloatImpl::GetDefaultFromChannel(FloatChannelHandle);
+			return DefaultValue.IsSet() ? DefaultValue.GetValue() : 0.f;
+		}
+		else
+		{
+			TOptional<double> DefaultValue = DoubleImpl::GetDefaultFromChannel(DoubleChannelHandle);
+			return DefaultValue.IsSet() ? (float)DefaultValue.GetValue() : 0.f;
+		}
 	}
 
 	/**
@@ -386,7 +615,14 @@ public:
 	UFUNCTION(BlueprintCallable, Category = "Sequencer|Keys", meta = (DisplayName = "Remove Default (Float)"))
 	void RemoveDefault()
 	{
-		RemoveDefaultFromChannel(ChannelHandle);
+		if (FloatChannelHandle.Get())
+		{
+			FloatImpl::RemoveDefaultFromChannel(FloatChannelHandle);
+		}
+		else
+		{
+			DoubleImpl::RemoveDefaultFromChannel(DoubleChannelHandle);
+		}
 	}
 
 	/**
@@ -395,10 +631,21 @@ public:
 	UFUNCTION(BlueprintCallable, Category = "Sequencer|Keys", meta = (DisplayName = "Has Default (Float)"))
 	bool HasDefault() const
 	{
-		return GetDefaultFromChannel(ChannelHandle).IsSet();
+		if (FloatChannelHandle.Get())
+		{
+			return FloatImpl::GetDefaultFromChannel(FloatChannelHandle).IsSet();
+		}
+		else
+		{
+			return DoubleImpl::GetDefaultFromChannel(DoubleChannelHandle).IsSet();
+		}
 	}
 public:
 	TWeakObjectPtr<UMovieSceneSequence> OwningSequence;
-	TMovieSceneChannelHandle<FMovieSceneFloatChannel> ChannelHandle;
 	TWeakObjectPtr<UMovieSceneSection> OwningSection;
-};
\ No newline at end of file
+private:
+	template<typename ChannelType, typename ScriptingChannelType>
+	friend void SetScriptingChannelHandle(ScriptingChannelType* ScriptingChannel, FMovieSceneChannelProxy& ChannelProxy, int32 ChannelIndex);
+	TMovieSceneChannelHandle<FMovieSceneFloatChannel> FloatChannelHandle;
+	TMovieSceneChannelHandle<FMovieSceneDoubleChannel> DoubleChannelHandle;
+};
diff --git a/Engine/Plugins/MovieScene/SequencerScripting/Source/SequencerScripting/Private/KeysAndChannels/MovieSceneScriptingInteger.h b/Engine/Plugins/MovieScene/SequencerScripting/Source/SequencerScripting/Private/KeysAndChannels/MovieSceneScriptingInteger.h
--- a/Engine/Plugins/MovieScene/SequencerScripting/Source/SequencerScripting/Private/KeysAndChannels/MovieSceneScriptingInteger.h
+++ b/Engine/Plugins/MovieScene/SequencerScripting/Source/SequencerScripting/Private/KeysAndChannels/MovieSceneScriptingInteger.h
@@ -62,9 +62,12 @@ public:
 };
 
 UCLASS(BlueprintType)
-class UMovieSceneScriptingIntegerChannel : public UMovieSceneScriptingChannel, public TMovieSceneScriptingChannel<FMovieSceneIntegerChannel, UMovieSceneScriptingIntegerKey, int32>
+class UMovieSceneScriptingIntegerChannel : public UMovieSceneScriptingChannel
 {
 	GENERATED_BODY()
+
+	using Impl = TMovieSceneScriptingChannel<FMovieSceneIntegerChannel, UMovieSceneScriptingIntegerKey, int32>;
+
 public:
 	/**
 	* Add a key to this channel. This initializes a new key and returns a reference to it.
@@ -77,7 +80,7 @@ public:
 	UFUNCTION(BlueprintCallable, Category = "Sequencer|Keys", meta = (DisplayName = "Add Key (Integer)"))
 	UMovieSceneScriptingIntegerKey* AddKey(const FFrameNumber& InTime, int32 NewValue, float SubFrame = 0.f, ESequenceTimeUnit TimeUnit = ESequenceTimeUnit::DisplayRate)
 	{
-		return AddKeyInChannel(ChannelHandle, OwningSequence, OwningSection, InTime, NewValue, SubFrame, TimeUnit, EMovieSceneKeyInterpolation::Auto);
+		return Impl::AddKeyInChannel(ChannelHandle, OwningSequence, OwningSection, InTime, NewValue, SubFrame, TimeUnit, EMovieSceneKeyInterpolation::Auto);
 	}
 
 	/**
@@ -86,7 +89,7 @@ public:
 	UFUNCTION(BlueprintCallable, Category = "Sequencer|Keys", meta = (DisplayName = "Remove Key (Integer)"))
 	virtual void RemoveKey(UMovieSceneScriptingKey* Key)
 	{
-		RemoveKeyFromChannel(ChannelHandle, Key);
+		Impl::RemoveKeyFromChannel(ChannelHandle, Key);
 	}
 
 	/**
@@ -97,7 +100,7 @@ public:
 	UFUNCTION(BlueprintCallable, Category = "Sequencer|Keys", meta = (DisplayName = "Get Keys (Integer)"))
 	virtual TArray<UMovieSceneScriptingKey*> GetKeys() const override
 	{
-		return GetKeysInChannel(ChannelHandle, OwningSequence, OwningSection);
+		return Impl::GetKeysInChannel(ChannelHandle, OwningSequence, OwningSection);
 	}
 
 	/**
@@ -107,7 +110,7 @@ public:
 	UFUNCTION(BlueprintCallable, Category = "Sequencer|Keys", meta = (DisplayName = "Set Default (Integer)"))
 	void SetDefault(int32 InDefaultValue)
 	{
-		SetDefaultInChannel(ChannelHandle, OwningSequence, OwningSection, InDefaultValue);
+		Impl::SetDefaultInChannel(ChannelHandle, OwningSequence, OwningSection, InDefaultValue);
 	}
 
 	/**
@@ -117,7 +120,7 @@ public:
 	UFUNCTION(BlueprintCallable, Category = "Sequencer|Keys", meta = (DisplayName = "Get Default (Integer)"))
 	int32 GetDefault() const
 	{
-		TOptional<int32> DefaultValue = GetDefaultFromChannel(ChannelHandle);
+		TOptional<int32> DefaultValue = Impl::GetDefaultFromChannel(ChannelHandle);
 		return DefaultValue.IsSet() ? DefaultValue.GetValue() : 0;
 	}
 
@@ -127,7 +130,7 @@ public:
 	UFUNCTION(BlueprintCallable, Category = "Sequencer|Keys", meta = (DisplayName = "Remove Default (Integer)"))
 	void RemoveDefault()
 	{
-		RemoveDefaultFromChannel(ChannelHandle);
+		Impl::RemoveDefaultFromChannel(ChannelHandle);
 	}
 
 	/**
@@ -136,10 +139,10 @@ public:
 	UFUNCTION(BlueprintCallable, Category = "Sequencer|Keys", meta = (DisplayName = "Has Default (Integer)"))
 	bool HasDefault() const
 	{
-		return GetDefaultFromChannel(ChannelHandle).IsSet();
+		return Impl::GetDefaultFromChannel(ChannelHandle).IsSet();
 	}
 public:
 	TWeakObjectPtr<UMovieSceneSequence> OwningSequence;
 	TMovieSceneChannelHandle<FMovieSceneIntegerChannel> ChannelHandle;
 	TWeakObjectPtr<UMovieSceneSection> OwningSection;
-};
\ No newline at end of file
+};
diff --git a/Engine/Plugins/MovieScene/SequencerScripting/Source/SequencerScripting/Private/KeysAndChannels/MovieSceneScriptingObjectPath.h b/Engine/Plugins/MovieScene/SequencerScripting/Source/SequencerScripting/Private/KeysAndChannels/MovieSceneScriptingObjectPath.h
--- a/Engine/Plugins/MovieScene/SequencerScripting/Source/SequencerScripting/Private/KeysAndChannels/MovieSceneScriptingObjectPath.h
+++ b/Engine/Plugins/MovieScene/SequencerScripting/Source/SequencerScripting/Private/KeysAndChannels/MovieSceneScriptingObjectPath.h
@@ -66,9 +66,12 @@ public:
 };
 
 UCLASS(BlueprintType)
-class UMovieSceneScriptingObjectPathChannel : public UMovieSceneScriptingChannel, public TMovieSceneScriptingChannel<FMovieSceneObjectPathChannel, UMovieSceneScriptingObjectPathKey, FMovieSceneObjectPathChannelKeyValue>
+class UMovieSceneScriptingObjectPathChannel : public UMovieSceneScriptingChannel
 {
 	GENERATED_BODY()
+
+	using Impl = TMovieSceneScriptingChannel<FMovieSceneObjectPathChannel, UMovieSceneScriptingObjectPathKey, FMovieSceneObjectPathChannelKeyValue>;
+
 public:
 	/**
 	 * Add a key to this channel. This initializes a new key and returns a reference to it.
@@ -82,7 +85,7 @@ public:
 	UMovieSceneScriptingObjectPathKey* AddKey(const FFrameNumber InTime, UObject* NewValue, float SubFrame = 0.f, ESequenceTimeUnit TimeUnit = ESequenceTimeUnit::DisplayRate)
 	{
 		FMovieSceneObjectPathChannelKeyValue ReferenceKey = FMovieSceneObjectPathChannelKeyValue(NewValue);
-		return AddKeyInChannel(ChannelHandle, OwningSequence, OwningSection,InTime, ReferenceKey, SubFrame, TimeUnit, EMovieSceneKeyInterpolation::Auto);
+		return Impl::AddKeyInChannel(ChannelHandle, OwningSequence, OwningSection,InTime, ReferenceKey, SubFrame, TimeUnit, EMovieSceneKeyInterpolation::Auto);
 	}
 
 	/**
@@ -91,7 +94,7 @@ public:
 	UFUNCTION(BlueprintCallable, Category = "Movie Scene | Keys", meta = (DisplayName = "Remove Key (Object Path)"))
 	virtual void RemoveKey(UMovieSceneScriptingKey* Key)
 	{
-		RemoveKeyFromChannel(ChannelHandle, Key);
+		Impl::RemoveKeyFromChannel(ChannelHandle, Key);
 	}
 
 	/**
@@ -102,7 +105,7 @@ public:
 	UFUNCTION(BlueprintCallable, Category = "Movie Scene | Keys", meta = (DisplayName = "Get Keys (Object Path)"))
 	virtual TArray<UMovieSceneScriptingKey*> GetKeys() const override
 	{
-		return GetKeysInChannel(ChannelHandle, OwningSequence, OwningSection);
+		return Impl::GetKeysInChannel(ChannelHandle, OwningSequence, OwningSection);
 	}
 
 	/**
@@ -179,4 +182,4 @@ public:
 	TWeakObjectPtr<UMovieSceneSequence> OwningSequence;
 	TWeakObjectPtr<UMovieSceneSection> OwningSection;
 	TMovieSceneChannelHandle<FMovieSceneObjectPathChannel> ChannelHandle;
-};
\ No newline at end of file
+};
diff --git a/Engine/Plugins/MovieScene/SequencerScripting/Source/SequencerScripting/Private/KeysAndChannels/MovieSceneScriptingString.h b/Engine/Plugins/MovieScene/SequencerScripting/Source/SequencerScripting/Private/KeysAndChannels/MovieSceneScriptingString.h
--- a/Engine/Plugins/MovieScene/SequencerScripting/Source/SequencerScripting/Private/KeysAndChannels/MovieSceneScriptingString.h
+++ b/Engine/Plugins/MovieScene/SequencerScripting/Source/SequencerScripting/Private/KeysAndChannels/MovieSceneScriptingString.h
@@ -62,9 +62,12 @@ public:
 };
 
 UCLASS(BlueprintType)
-class UMovieSceneScriptingStringChannel : public UMovieSceneScriptingChannel, public TMovieSceneScriptingChannel<FMovieSceneStringChannel, UMovieSceneScriptingStringKey, FString>
+class UMovieSceneScriptingStringChannel : public UMovieSceneScriptingChannel
 {
 	GENERATED_BODY()
+
+	using Impl = TMovieSceneScriptingChannel<FMovieSceneStringChannel, UMovieSceneScriptingStringKey, FString>;
+
 public:
 	/**
 	* Add a key to this channel. This initializes a new key and returns a reference to it.
@@ -78,7 +81,7 @@ public:
 	UFUNCTION(BlueprintCallable, Category = "Sequencer|Keys", meta = (DisplayName = "Add Key (String)"))
 	UMovieSceneScriptingStringKey* AddKey(const FFrameNumber& InTime, FString NewValue, float SubFrame = 0.f, ESequenceTimeUnit TimeUnit = ESequenceTimeUnit::DisplayRate)
 	{
-		return AddKeyInChannel(ChannelHandle, OwningSequence, OwningSection,InTime, NewValue, SubFrame, TimeUnit, EMovieSceneKeyInterpolation::Auto);
+		return Impl::AddKeyInChannel(ChannelHandle, OwningSequence, OwningSection,InTime, NewValue, SubFrame, TimeUnit, EMovieSceneKeyInterpolation::Auto);
 	}
 
 	/**
@@ -87,7 +90,7 @@ public:
 	UFUNCTION(BlueprintCallable, Category = "Sequencer|Keys", meta = (DisplayName = "Remove Key (String)"))
 	virtual void RemoveKey(UMovieSceneScriptingKey* Key)
 	{
-		RemoveKeyFromChannel(ChannelHandle, Key);
+		Impl::RemoveKeyFromChannel(ChannelHandle, Key);
 	}
 
 	/**
@@ -98,7 +101,7 @@ public:
 	UFUNCTION(BlueprintCallable, Category = "Sequencer|Keys", meta = (DisplayName = "Get Keys (String)"))
 	virtual TArray<UMovieSceneScriptingKey*> GetKeys() const override
 	{
-		return GetKeysInChannel(ChannelHandle, OwningSequence, OwningSection);
+		return Impl::GetKeysInChannel(ChannelHandle, OwningSequence, OwningSection);
 	}
 
 	/**
@@ -108,7 +111,7 @@ public:
 	UFUNCTION(BlueprintCallable, Category = "Sequencer|Keys", meta = (DisplayName = "Set Default (String)"))
 	void SetDefault(FString InDefaultValue)
 	{
-		SetDefaultInChannel(ChannelHandle, OwningSequence, OwningSection, InDefaultValue);
+		Impl::SetDefaultInChannel(ChannelHandle, OwningSequence, OwningSection, InDefaultValue);
 	}
 
 	/**
@@ -118,7 +121,7 @@ public:
 	UFUNCTION(BlueprintCallable, Category = "Sequencer|Keys", meta = (DisplayName = "Get Default (String)"))
 	FString GetDefault() const
 	{
-		TOptional<FString> DefaultValue = GetDefaultFromChannel(ChannelHandle);
+		TOptional<FString> DefaultValue = Impl::GetDefaultFromChannel(ChannelHandle);
 		return DefaultValue.IsSet() ? DefaultValue.GetValue() : FString();
 	}
 
@@ -128,7 +131,7 @@ public:
 	UFUNCTION(BlueprintCallable, Category = "Sequencer|Keys", meta = (DisplayName = "Remove Default (String)"))
 	void RemoveDefault()
 	{
-		RemoveDefaultFromChannel(ChannelHandle);
+		Impl::RemoveDefaultFromChannel(ChannelHandle);
 	}
 
 	/**
@@ -137,10 +140,10 @@ public:
 	UFUNCTION(BlueprintCallable, Category = "Sequencer|Keys", meta = (DisplayName = "Has Default (String)"))
 	bool HasDefault() const
 	{
-		return GetDefaultFromChannel(ChannelHandle).IsSet();
+		return Impl::GetDefaultFromChannel(ChannelHandle).IsSet();
 	}
 public:
 	TWeakObjectPtr<UMovieSceneSequence> OwningSequence;
 	TMovieSceneChannelHandle<FMovieSceneStringChannel> ChannelHandle;
 	TWeakObjectPtr<UMovieSceneSection> OwningSection;
-};
\ No newline at end of file
+};
diff --git a/Engine/Plugins/MovieScene/SequencerScripting/Source/SequencerScripting/Public/ExtensionLibraries/MovieSceneSectionExtensions.h b/Engine/Plugins/MovieScene/SequencerScripting/Source/SequencerScripting/Public/ExtensionLibraries/MovieSceneSectionExtensions.h
--- a/Engine/Plugins/MovieScene/SequencerScripting/Source/SequencerScripting/Public/ExtensionLibraries/MovieSceneSectionExtensions.h
+++ b/Engine/Plugins/MovieScene/SequencerScripting/Source/SequencerScripting/Public/ExtensionLibraries/MovieSceneSectionExtensions.h
@@ -157,23 +157,53 @@ public:
 	* Transforms containing 9 float channels to represent Translation/Rotation/Scale), and a section may have mixed
 	* channel types.
 	*
+	* For backwards compatibility reasons, all double channels will be wrapped inside float scripting channels.
+	* Use GetAllChannels to get proper double channels.
+	*
 	* @param Section       The section to use.
 	* @return An array containing any key channels that match the type specified
 	*/
+	UE_DEPRECATED(5.0, "Please use GetAllChannels instead")
 	UFUNCTION(BlueprintCallable, Category = "Movie Scene | Section", meta=(ScriptMethod, DevelopmentOnly))
 	static TArray<UMovieSceneScriptingChannel*> GetChannels(UMovieSceneSection* Section);
 
+	/**
+	* Find all channels that belong to the specified UMovieSceneSection. Some sections have many channels (such as
+	* Transforms containing 9 double channels to represent Translation/Rotation/Scale), and a section may have mixed
+	* channel types.
+	*
+	* @param Section       The section to use.
+	* @return An array containing any key channels that match the type specified
+	*/
+	UFUNCTION(BlueprintCallable, Category = "Movie Scene | Section", meta=(ScriptMethod, DevelopmentOnly))
+	static TArray<UMovieSceneScriptingChannel*> GetAllChannels(UMovieSceneSection* Section);
+
 	/**
 	* Find all channels that belong to the specified UMovieSceneSection that match the specific type. This will filter out any children who do not inherit
 	* from the specified type for you.
 	*
+	* For backwards compatibility reasons, all double channels will be wrapped inside float scripting channels.
+	* Use GetAllChannels to get proper double channels.
+	*
 	* @param Section        The section to use.
 	* @param ChannelType	The class type to look for.
 	* @return An array containing any key channels that match the type specified
 	*/
+	UE_DEPRECATED(5.0, "Please use GetChannelsByType instead")
 	UFUNCTION(BlueprintCallable, Category = "Movie Scene | Section", meta = (ScriptMethod, DeterminesOutputType="TrackType", DevelopmentOnly))
 	static TArray<UMovieSceneScriptingChannel*> FindChannelsByType(UMovieSceneSection* Section, TSubclassOf<UMovieSceneScriptingChannel> ChannelType);
 
+	/**
+	* Find all channels that belong to the specified UMovieSceneSection that match the specific type. This will filter out any children who do not inherit
+	* from the specified type for you.
+	*
+	* @param Section        The section to use.
+	* @param ChannelType	The class type to look for.
+	* @return An array containing any key channels that match the type specified
+	*/
+	UFUNCTION(BlueprintCallable, Category = "Movie Scene | Section", meta = (ScriptMethod, DeterminesOutputType="TrackType", DevelopmentOnly))
+	static TArray<UMovieSceneScriptingChannel*> GetChannelsByType(UMovieSceneSection* Section, TSubclassOf<UMovieSceneScriptingChannel> ChannelType);
+
 	/**
 	 * Get the frame in the space of its parent sequence
 	 *
@@ -184,4 +214,4 @@ public:
 	 */
 	UFUNCTION(BlueprintCallable, Category = "Movie Scene | Section", meta = (ScriptMethod))
 	static int32 GetParentSequenceFrame(UMovieSceneSubSection* Section, int32 InFrame, UMovieSceneSequence* ParentSequence);
-};
\ No newline at end of file
+};
