No unity build fixes from partner
#rnx

diff --git a/Engine/Plugins/Experimental/MeshModelingToolsetExp/Source/MeshModelingToolsEditorOnly/Private/AddPivotActorTool.cpp b/Engine/Plugins/Experimental/MeshModelingToolsetExp/Source/MeshModelingToolsEditorOnly/Private/AddPivotActorTool.cpp
--- a/Engine/Plugins/Experimental/MeshModelingToolsetExp/Source/MeshModelingToolsEditorOnly/Private/AddPivotActorTool.cpp
+++ b/Engine/Plugins/Experimental/MeshModelingToolsetExp/Source/MeshModelingToolsEditorOnly/Private/AddPivotActorTool.cpp
@@ -7,6 +7,7 @@
 #include "BaseGizmos/TransformGizmoUtil.h"
 #include "BaseGizmos/TransformProxy.h"
 #include "Components/BillboardComponent.h"
+#include "Editor.h"
 #include "Mechanics/DragAlignmentMechanic.h"
 #include "ModelingToolTargetUtil.h"
 #include "Selection/ToolSelectionUtil.h"
diff --git a/Engine/Plugins/MovieScene/SequencerScripting/Source/SequencerScripting/Private/KeysAndChannels/MovieSceneScriptingFloat.h b/Engine/Plugins/MovieScene/SequencerScripting/Source/SequencerScripting/Private/KeysAndChannels/MovieSceneScriptingFloat.h
--- a/Engine/Plugins/MovieScene/SequencerScripting/Source/SequencerScripting/Private/KeysAndChannels/MovieSceneScriptingFloat.h
+++ b/Engine/Plugins/MovieScene/SequencerScripting/Source/SequencerScripting/Private/KeysAndChannels/MovieSceneScriptingFloat.h
@@ -5,6 +5,7 @@
 #include "CoreTypes.h"
 #include "UObject/ObjectMacros.h"
 #include "MovieSceneScriptingChannel.h"
+#include "Algo/Transform.h"
 #include "Channels/MovieSceneChannelTraits.h"
 #include "Channels/MovieSceneFloatChannel.h"
 #include "Channels/MovieSceneDoubleChannel.h"
diff --git a/Engine/Plugins/Runtime/Interchange/Source/Import/Private/Scene/InterchangeActorFactory.cpp b/Engine/Plugins/Runtime/Interchange/Source/Import/Private/Scene/InterchangeActorFactory.cpp
--- a/Engine/Plugins/Runtime/Interchange/Source/Import/Private/Scene/InterchangeActorFactory.cpp
+++ b/Engine/Plugins/Runtime/Interchange/Source/Import/Private/Scene/InterchangeActorFactory.cpp
@@ -7,6 +7,7 @@
 #include "InterchangeSourceData.h"
 
 #include "Components/SceneComponent.h"
+#include "Components/SkeletalMeshComponent.h"
 #include "Engine/Engine.h"
 #include "Engine/StaticMesh.h"
 #include "Engine/World.h"
diff --git a/Engine/Plugins/WorldBuilding/Source/WorldBuilding/Private/Replay/WorldBuildingReplaySpectator.cpp b/Engine/Plugins/WorldBuilding/Source/WorldBuilding/Private/Replay/WorldBuildingReplaySpectator.cpp
--- a/Engine/Plugins/WorldBuilding/Source/WorldBuilding/Private/Replay/WorldBuildingReplaySpectator.cpp
+++ b/Engine/Plugins/WorldBuilding/Source/WorldBuilding/Private/Replay/WorldBuildingReplaySpectator.cpp
@@ -3,6 +3,7 @@
 #include "WorldBuildingReplaySpectator.h"
 #include "WorldPartition/WorldPartition.h"
 #include "WorldBuildingReplaySpectatorHUD.h"
+#include "Algo/Transform.h"
 #include "Engine/DemoNetDriver.h"
 #include "GameFramework/PlayerInput.h"
 
diff --git a/Engine/Source/Runtime/VulkanRHI/Private/VulkanRenderpass.h b/Engine/Source/Runtime/VulkanRHI/Private/VulkanRenderpass.h
--- a/Engine/Source/Runtime/VulkanRHI/Private/VulkanRenderpass.h
+++ b/Engine/Source/Runtime/VulkanRHI/Private/VulkanRenderpass.h
@@ -7,7 +7,7 @@
 #pragma once
 
 #include "CoreMinimal.h"
-#include "RHIResources.h"
+#include "RHI.h"
 #include "VulkanRHIPrivate.h"
 #include "VulkanResources.h"
 #include "VulkanPendingState.h"
