Sequencer: LevelSequenceEditor now depends on SequencerScripting

diff --git a/Engine/Plugins/MovieScene/LevelSequenceEditor/LevelSequenceEditor.uplugin b/Engine/Plugins/MovieScene/LevelSequenceEditor/LevelSequenceEditor.uplugin
--- a/Engine/Plugins/MovieScene/LevelSequenceEditor/LevelSequenceEditor.uplugin
+++ b/Engine/Plugins/MovieScene/LevelSequenceEditor/LevelSequenceEditor.uplugin
@@ -21,5 +21,11 @@
 			"Type" : "Editor",
 			"LoadingPhase" : "PostEngineInit"
 		}
+	],
+	"Plugins": [
+		{
+			"Name": "SequencerScripting",
+			"Enabled": true
+		}
 	]
 }
\ No newline at end of file
diff --git a/Engine/Plugins/MovieScene/SequencerScripting/SequencerScripting.uplugin b/Engine/Plugins/MovieScene/SequencerScripting/SequencerScripting.uplugin
--- a/Engine/Plugins/MovieScene/SequencerScripting/SequencerScripting.uplugin
+++ b/Engine/Plugins/MovieScene/SequencerScripting/SequencerScripting.uplugin
@@ -18,11 +18,6 @@
 			"Name": "PythonScriptPlugin",
 			"Enabled": true
 		},
-		{
-			"Name": "TemplateSequence",
-			"Enabled": true
-		}
-
 	],
 	"Modules": [
 		{
diff --git a/Engine/Plugins/MovieScene/SequencerScripting/Source/SequencerScriptingEditor/Private/SequencerTools.cpp b/Engine/Plugins/MovieScene/SequencerScripting/Source/SequencerScriptingEditor/Private/SequencerTools.cpp
--- a/Engine/Plugins/MovieScene/SequencerScripting/Source/SequencerScriptingEditor/Private/SequencerTools.cpp
+++ b/Engine/Plugins/MovieScene/SequencerScripting/Source/SequencerScriptingEditor/Private/SequencerTools.cpp
@@ -25,10 +25,6 @@
 #include "ScopedTransaction.h"
 #include "Exporters/AnimSeqExportOption.h"
 
-#include "TemplateSequenceActor.h"
-#include "TemplateSequencePlayer.h"
-#include "TemplateSequence.h"
-
 #include "K2Node_CustomEvent.h"
 #include "BlueprintFunctionNodeSpawner.h"
 #include "BlueprintActionMenuItem.h"
@@ -235,19 +231,6 @@ bool USequencerToolsFunctionLibrary::ExportLevelSequenceFBX(UWorld* World, ULeve
 	return bSuccess;
 }
 
-bool USequencerToolsFunctionLibrary::ExportTemplateSequenceFBX(UWorld* World, UTemplateSequence* Sequence, const TArray<FSequencerBindingProxy>& InBindings, UFbxExportOption* OverrideOptions, const FString& InFBXFileName)
-{
-	ATemplateSequenceActor* OutActor;
-	FMovieSceneSequencePlaybackSettings Settings;
-	UTemplateSequencePlayer* Player = UTemplateSequencePlayer::CreateTemplateSequencePlayer(World, Sequence, Settings, OutActor);
-	Player->Initialize(Sequence, World, Settings);
-	
-	bool bSuccess = ExportFBXInternal(World, Sequence, InBindings, TArray<UMovieSceneTrack*>(), OverrideOptions, InFBXFileName, Player);
-	World->DestroyActor(OutActor);
-
-	return bSuccess;
-}
-
 static USkeletalMeshComponent* GetSkelMeshComponent(IMovieScenePlayer* Player, const FSequencerBindingProxy& Binding)
 {
 	FMovieSceneSequenceIDRef Template = MovieSceneSequenceID::Root;
@@ -496,22 +479,6 @@ bool USequencerToolsFunctionLibrary::ImportLevelSequenceFBX(UWorld* World, ULeve
 	return bSuccess;
 }
 
-bool USequencerToolsFunctionLibrary::ImportTemplateSequenceFBX(UWorld* World, UTemplateSequence* Sequence, const TArray<FSequencerBindingProxy>& InBindings, UMovieSceneUserImportFBXSettings* ImportFBXSettings, const FString&  ImportFilename)
-{
-	ImportFBXSettings->bCreateCameras = false;
-	ImportFBXSettings->bMatchByNameOnly = false;
-
-	ATemplateSequenceActor* OutActor;
-	FMovieSceneSequencePlaybackSettings Settings;
-	UTemplateSequencePlayer* Player = UTemplateSequencePlayer::CreateTemplateSequencePlayer(World, Sequence, Settings, OutActor);
-	Player->Initialize(Sequence, World, Settings);
-
-	bool bSuccess = ImportFBXInternal(World, Sequence, InBindings, ImportFBXSettings, ImportFilename, Player);
-	World->DestroyActor(OutActor);
-
-	return bSuccess;
-}
-
 bool USequencerToolsFunctionLibrary::ImportFBXToControlRig(UWorld* World, ULevelSequence* Sequence, const FString& ControlRigTrackName, const TArray<FString>& ControlRigNames,
 	UMovieSceneUserImportFBXControlRigSettings* ImportFBXControlRigSettings,
 	const FString& ImportFilename)
diff --git a/Engine/Plugins/MovieScene/SequencerScripting/Source/SequencerScriptingEditor/Public/SequencerTools.h b/Engine/Plugins/MovieScene/SequencerScripting/Source/SequencerScriptingEditor/Public/SequencerTools.h
--- a/Engine/Plugins/MovieScene/SequencerScripting/Source/SequencerScriptingEditor/Public/SequencerTools.h
+++ b/Engine/Plugins/MovieScene/SequencerScripting/Source/SequencerScriptingEditor/Public/SequencerTools.h
@@ -17,7 +17,6 @@ class UAnimSequenceExportOption;
 class UAnimSequence;
 class UPoseAsset;
 class UMovieSceneSequencePlayer;
-class UTemplateSequence;
 
 class UMovieSceneEventSectionBase;
 class UK2Node_CustomEvent;
@@ -128,9 +127,6 @@ public:
 	UFUNCTION(BlueprintCallable, Category = "Editor Scripting | Sequencer Tools | FBX")
 	static bool ExportLevelSequenceFBX(UWorld* InWorld, ULevelSequence* InSequence, const TArray<FSequencerBindingProxy>& InBindings, const TArray<UMovieSceneTrack*>& InMasterTracks, UFbxExportOption* OverrideOptions, const FString& InFBXFileName);
 
-	UFUNCTION(BlueprintCallable, Category = "Editor Scripting | Sequencer Tools | FBX")
-	static bool ExportTemplateSequenceFBX(UWorld* InWorld, UTemplateSequence* InSequence, const TArray<FSequencerBindingProxy>& InBindings, UFbxExportOption* OverrideOptions, const FString& InFBXFileName);
-
 	UE_DEPRECATED(4.27, "Please use ExportLevelSequenceFBX instead")
 	static bool ExportFBX(UWorld* InWorld, ULevelSequence* InSequence, const TArray<FSequencerBindingProxy>& InBindings, UFbxExportOption* OverrideOptions,const FString& InFBXFileName) { return ExportLevelSequenceFBX(InWorld, InSequence, InBindings, InSequence->GetMovieScene()->GetMasterTracks(), OverrideOptions, InFBXFileName); }
 
@@ -160,9 +156,6 @@ public:
 	UFUNCTION(BlueprintCallable, Category = "Editor Scripting | Sequencer Tools | FBX")
 	static bool ImportLevelSequenceFBX(UWorld* InWorld, ULevelSequence* InSequence, const TArray<FSequencerBindingProxy>& InBindings, UMovieSceneUserImportFBXSettings* InImportFBXSettings, const FString& InImportFilename);
 
-	UFUNCTION(BlueprintCallable, Category = "Editor Scripting | Sequencer Tools | FBX")
-	static bool ImportTemplateSequenceFBX(UWorld* InWorld, UTemplateSequence* InSequence, const TArray<FSequencerBindingProxy>& InBindings, UMovieSceneUserImportFBXSettings* InImportFBXSettings, const FString& InImportFilename);
-
 	UE_DEPRECATED(4.27, "Please use ImportLevelSequenceFBX instead")
 	static bool ImportFBX(UWorld* InWorld, ULevelSequence* InSequence, const TArray<FSequencerBindingProxy>& InBindings, UMovieSceneUserImportFBXSettings* InImportFBXSettings, const FString& InImportFilename) { return ImportLevelSequenceFBX(InWorld, InSequence, InBindings, InImportFBXSettings, InImportFilename); }
 
diff --git a/Engine/Plugins/MovieScene/SequencerScripting/Source/SequencerScriptingEditor/SequencerScriptingEditor.Build.cs b/Engine/Plugins/MovieScene/SequencerScripting/Source/SequencerScriptingEditor/SequencerScriptingEditor.Build.cs
--- a/Engine/Plugins/MovieScene/SequencerScripting/Source/SequencerScriptingEditor/SequencerScriptingEditor.Build.cs
+++ b/Engine/Plugins/MovieScene/SequencerScripting/Source/SequencerScriptingEditor/SequencerScriptingEditor.Build.cs
@@ -31,9 +31,7 @@ public class SequencerScriptingEditor : ModuleRules
 				"MovieSceneTools",
                 "MovieSceneTracks",
 				"CinematicCamera",
-				"SequencerScripting",
-				"TemplateSequence",
-				"TemplateSequenceEditor",
+				"SequencerScripting"
 			}
 		);
 
