Sequencer: Added SetMaterialIndex and GetMaterialIndex for PrimitiveMaterialTrack and ComponentMaterialTrack

diff --git a/Engine/Plugins/MovieScene/SequencerScripting/Source/SequencerScripting/Private/ExtensionLibraries/MovieSceneMaterialTrackExtensions.cpp b/Engine/Plugins/MovieScene/SequencerScripting/Source/SequencerScripting/Private/ExtensionLibraries/MovieSceneMaterialTrackExtensions.cpp
new file mode 100644
--- /dev/null
+++ b/Engine/Plugins/MovieScene/SequencerScripting/Source/SequencerScripting/Private/ExtensionLibraries/MovieSceneMaterialTrackExtensions.cpp
@@ -0,0 +1,29 @@
+// Copyright Epic Games, Inc. All Rights Reserved.
+
+#include "ExtensionLibraries/MovieSceneMaterialTrackExtensions.h"
+#include "Tracks/MovieSceneMaterialTrack.h"
+
+void UMovieSceneMaterialTrackExtensions::SetMaterialIndex(UMovieSceneComponentMaterialTrack* Track, const int32 MaterialIndex)
+{
+	if (!Track)
+	{
+		FFrame::KismetExecutionMessage(TEXT("Cannot call SetMaterialIndex on a null track"), ELogVerbosity::Error);
+		return;
+	}
+
+	Track->Modify();
+	Track->SetMaterialIndex(MaterialIndex);
+}
+
+int32 UMovieSceneMaterialTrackExtensions::GetMaterialIndex(UMovieSceneComponentMaterialTrack* Track)
+{
+	if (!Track)
+	{
+		FFrame::KismetExecutionMessage(TEXT("Cannot call GetMaterialIndex on a null track"), ELogVerbosity::Error);
+		return -1;
+	}
+
+	return Track->GetMaterialIndex();
+}
+
+
diff --git a/Engine/Plugins/MovieScene/SequencerScripting/Source/SequencerScripting/Private/ExtensionLibraries/MovieScenePrimitiveMaterialTrackExtensions.cpp b/Engine/Plugins/MovieScene/SequencerScripting/Source/SequencerScripting/Private/ExtensionLibraries/MovieScenePrimitiveMaterialTrackExtensions.cpp
new file mode 100644
--- /dev/null
+++ b/Engine/Plugins/MovieScene/SequencerScripting/Source/SequencerScripting/Private/ExtensionLibraries/MovieScenePrimitiveMaterialTrackExtensions.cpp
@@ -0,0 +1,29 @@
+// Copyright Epic Games, Inc. All Rights Reserved.
+
+#include "ExtensionLibraries/MovieScenePrimitiveMaterialTrackExtensions.h"
+#include "Tracks/MovieScenePrimitiveMaterialTrack.h"
+
+void UMovieScenePrimitiveMaterialTrackExtensions::SetMaterialIndex(UMovieScenePrimitiveMaterialTrack* Track, const int32 MaterialIndex)
+{
+	if (!Track)
+	{
+		FFrame::KismetExecutionMessage(TEXT("Cannot call SetMaterialIndex on a null track"), ELogVerbosity::Error);
+		return;
+	}
+
+	Track->Modify();
+	Track->SetMaterialIndex(MaterialIndex);
+}
+
+int32 UMovieScenePrimitiveMaterialTrackExtensions::GetMaterialIndex(UMovieScenePrimitiveMaterialTrack* Track)
+{
+	if (!Track)
+	{
+		FFrame::KismetExecutionMessage(TEXT("Cannot call GetMaterialIndex on a null track"), ELogVerbosity::Error);
+		return -1;
+	}
+
+	return Track->GetMaterialIndex();
+}
+
+
diff --git a/Engine/Plugins/MovieScene/SequencerScripting/Source/SequencerScripting/Public/ExtensionLibraries/MovieSceneMaterialTrackExtensions.h b/Engine/Plugins/MovieScene/SequencerScripting/Source/SequencerScripting/Public/ExtensionLibraries/MovieSceneMaterialTrackExtensions.h
new file mode 100644
--- /dev/null
+++ b/Engine/Plugins/MovieScene/SequencerScripting/Source/SequencerScripting/Public/ExtensionLibraries/MovieSceneMaterialTrackExtensions.h
@@ -0,0 +1,38 @@
+// Copyright Epic Games, Inc. All Rights Reserved.
+
+#pragma once
+
+#include "CoreTypes.h"
+#include "UObject/ObjectMacros.h"
+#include "Kismet/BlueprintFunctionLibrary.h"
+
+#include "MovieSceneMaterialTrackExtensions.generated.h"
+
+class UMovieSceneMaterialTrack;
+
+/**
+ * Function library containing methods that should be hoisted onto UMovieSceneMaterialTrack for scripting
+ */
+UCLASS()
+class UMovieSceneMaterialTrackExtensions : public UBlueprintFunctionLibrary
+{
+	GENERATED_BODY()
+
+public:
+
+	/**
+	 * Set material index of the component that is animated by the material track.
+	 * @param Track The track to use
+	 * @param MaterialIndex The desired material index to animate. Values of < 0 or >= NumMaterials will be silently ignored and evaluation will fail.
+	 */
+	UFUNCTION(BlueprintCallable, Category = "Sequencer|Track", meta = (ScriptMethod))
+	static void SetMaterialIndex(UMovieSceneComponentMaterialTrack* Track, const int32 MaterialIndex);
+
+	/**
+	 * Get material index of the component that is animated by the material track.
+	 * @param Track The track to use
+	 * @return The material index.
+	 */
+	UFUNCTION(BlueprintCallable, Category = "Sequencer|Track", meta=(ScriptMethod))
+	static int32 GetMaterialIndex(UMovieSceneComponentMaterialTrack* Track);
+};
diff --git a/Engine/Plugins/MovieScene/SequencerScripting/Source/SequencerScripting/Public/ExtensionLibraries/MovieScenePrimitiveMaterialTrackExtensions.h b/Engine/Plugins/MovieScene/SequencerScripting/Source/SequencerScripting/Public/ExtensionLibraries/MovieScenePrimitiveMaterialTrackExtensions.h
new file mode 100644
--- /dev/null
+++ b/Engine/Plugins/MovieScene/SequencerScripting/Source/SequencerScripting/Public/ExtensionLibraries/MovieScenePrimitiveMaterialTrackExtensions.h
@@ -0,0 +1,38 @@
+// Copyright Epic Games, Inc. All Rights Reserved.
+
+#pragma once
+
+#include "CoreTypes.h"
+#include "UObject/ObjectMacros.h"
+#include "Kismet/BlueprintFunctionLibrary.h"
+
+#include "MovieScenePrimitiveMaterialTrackExtensions.generated.h"
+
+class UMovieScenePrimitiveMaterialTrack;
+
+/**
+ * Function library containing methods that should be hoisted onto UMovieScenePrimitiveMaterialTrack for scripting
+ */
+UCLASS()
+class UMovieScenePrimitiveMaterialTrackExtensions : public UBlueprintFunctionLibrary
+{
+	GENERATED_BODY()
+
+public:
+
+	/**
+	 * Set material index of the element that is animated by the primitive material track.
+	 * @param Track The track to use
+	 * @param MaterialIndex The desired material index to animate. Values of < 0 or >= NumMaterials will be silently ignored and evaluation will fail.
+	 */
+	UFUNCTION(BlueprintCallable, Category = "Sequencer|Track", meta = (ScriptMethod))
+	static void SetMaterialIndex(UMovieScenePrimitiveMaterialTrack* Track, const int32 MaterialIndex);
+
+	/**
+	 * Get material index of the element that is animated by the primitive material track.
+	 * @param Track The track to use
+	 * @return The material index.
+	 */
+	UFUNCTION(BlueprintCallable, Category = "Sequencer|Track", meta=(ScriptMethod))
+	static int32 GetMaterialIndex(UMovieScenePrimitiveMaterialTrack* Track);
+};
diff --git a/Engine/Source/Editor/MovieSceneTools/Private/TrackEditors/PrimitiveMaterialTrackEditor.cpp b/Engine/Source/Editor/MovieSceneTools/Private/TrackEditors/PrimitiveMaterialTrackEditor.cpp
--- a/Engine/Source/Editor/MovieSceneTools/Private/TrackEditors/PrimitiveMaterialTrackEditor.cpp
+++ b/Engine/Source/Editor/MovieSceneTools/Private/TrackEditors/PrimitiveMaterialTrackEditor.cpp
@@ -34,7 +34,7 @@ void FPrimitiveMaterialTrackEditor::ConstructObjectBindingTrackMenu(FMenuBuilder
 	auto GetMaterialIndexForTrack = [](UMovieSceneTrack* InTrack)
 	{
 		UMovieScenePrimitiveMaterialTrack* MaterialTrack = Cast<UMovieScenePrimitiveMaterialTrack>(InTrack);
-		return MaterialTrack ? MaterialTrack->MaterialIndex : INDEX_NONE;
+		return MaterialTrack ? MaterialTrack->GetMaterialIndex() : INDEX_NONE;
 	};
 
 	int32 MinNumMaterials = TNumericLimits<int32>::Max();
@@ -87,7 +87,7 @@ void FPrimitiveMaterialTrackEditor::CreateTrackForElement(TArray<FGuid> ObjectBi
 	for (FGuid ObjectBindingID : ObjectBindingIDs)
 	{
 		UMovieScenePrimitiveMaterialTrack* NewTrack = MovieScene->AddTrack<UMovieScenePrimitiveMaterialTrack>(ObjectBindingID);
-		NewTrack->MaterialIndex = MaterialIndex;
+		NewTrack->SetMaterialIndex(MaterialIndex);
 		NewTrack->SetDisplayName(FText::Format(LOCTEXT("MaterialTrackName_Format", "Material Element {0}"), FText::AsNumber(MaterialIndex)));
 
 		NewTrack->AddSection(*NewTrack->CreateNewSection());
diff --git a/Engine/Source/Runtime/MovieSceneTracks/Private/Evaluation/MovieScenePrimitiveMaterialTemplate.cpp b/Engine/Source/Runtime/MovieSceneTracks/Private/Evaluation/MovieScenePrimitiveMaterialTemplate.cpp
--- a/Engine/Source/Runtime/MovieSceneTracks/Private/Evaluation/MovieScenePrimitiveMaterialTemplate.cpp
+++ b/Engine/Source/Runtime/MovieSceneTracks/Private/Evaluation/MovieScenePrimitiveMaterialTemplate.cpp
@@ -83,7 +83,7 @@ struct FPrimitiveMaterialExecToken : IMovieSceneExecutionToken
 
 
 FMovieScenePrimitiveMaterialTemplate::FMovieScenePrimitiveMaterialTemplate(const UMovieScenePrimitiveMaterialSection& Section, const UMovieScenePrimitiveMaterialTrack& Track)
-	: MaterialIndex(Track.MaterialIndex)
+	: MaterialIndex(Track.GetMaterialIndex())
 	, MaterialChannel(Section.MaterialChannel)
 {}
 
diff --git a/Engine/Source/Runtime/MovieSceneTracks/Private/Tracks/MovieScenePrimitiveMaterialTrack.cpp b/Engine/Source/Runtime/MovieSceneTracks/Private/Tracks/MovieScenePrimitiveMaterialTrack.cpp
--- a/Engine/Source/Runtime/MovieSceneTracks/Private/Tracks/MovieScenePrimitiveMaterialTrack.cpp
+++ b/Engine/Source/Runtime/MovieSceneTracks/Private/Tracks/MovieScenePrimitiveMaterialTrack.cpp
@@ -35,4 +35,14 @@ void UMovieScenePrimitiveMaterialTrack::PostCompile(FMovieSceneEvaluationTrack&
 {
 	// Must evaluate before material parameter tracks
 	OutTrack.SetEvaluationPriority(UMovieSceneComponentMaterialTrack::EvaluationPriority + 1);
+}
+
+int32 UMovieScenePrimitiveMaterialTrack::GetMaterialIndex() const
+{
+	return MaterialIndex;
+}
+
+void UMovieScenePrimitiveMaterialTrack::SetMaterialIndex(int32 InMaterialIndex)
+{
+	MaterialIndex = InMaterialIndex;
 }
\ No newline at end of file
diff --git a/Engine/Source/Runtime/MovieSceneTracks/Public/Tracks/MovieScenePrimitiveMaterialTrack.h b/Engine/Source/Runtime/MovieSceneTracks/Public/Tracks/MovieScenePrimitiveMaterialTrack.h
--- a/Engine/Source/Runtime/MovieSceneTracks/Public/Tracks/MovieScenePrimitiveMaterialTrack.h
+++ b/Engine/Source/Runtime/MovieSceneTracks/Public/Tracks/MovieScenePrimitiveMaterialTrack.h
@@ -14,14 +14,21 @@ public:
 
 	GENERATED_BODY()
 
-	UPROPERTY()
-	int32 MaterialIndex;
-
 	UMovieScenePrimitiveMaterialTrack(const FObjectInitializer& ObjInit);
 
+	/* Set the material index that this track is assigned to */
+	MOVIESCENETRACKS_API void SetMaterialIndex(int32 MaterialIndex);
+
+	/* Get the material index that this track is assigned to */
+	MOVIESCENETRACKS_API int32 GetMaterialIndex() const;
+
 	/*~ UMovieSceneTrack interface */
 	virtual UMovieSceneSection* CreateNewSection() override;
 	virtual bool SupportsType(TSubclassOf<UMovieSceneSection> SectionClass) const override;
 	virtual FMovieSceneEvalTemplatePtr CreateTemplateForSection(const UMovieSceneSection& InSection) const override;
 	virtual void PostCompile(FMovieSceneEvaluationTrack& OutTrack, const FMovieSceneTrackCompilerArgs& Args) const override;
+
+private:
+	UPROPERTY()
+	int32 MaterialIndex;
 };
