Sequencer: Get descendant bindings when importing/exporting fbx otherwise components will not be imported/exported if you only specify the actor binding

diff --git a/Engine/Plugins/MovieScene/SequencerScripting/Source/SequencerScriptingEditor/Private/SequencerTools.cpp b/Engine/Plugins/MovieScene/SequencerScripting/Source/SequencerScriptingEditor/Private/SequencerTools.cpp
--- a/Engine/Plugins/MovieScene/SequencerScripting/Source/SequencerScriptingEditor/Private/SequencerTools.cpp
+++ b/Engine/Plugins/MovieScene/SequencerScripting/Source/SequencerScriptingEditor/Private/SequencerTools.cpp
@@ -177,6 +177,22 @@ TArray<FSequencerBoundObjects> USequencerToolsFunctionLibrary::GetObjectBindings
 	return BoundObjects;
 }
 
+void GatherDescendantBindings(const FSequencerBindingProxy& Binding, UMovieSceneSequence* Sequence, TArray<FSequencerBindingProxy>& AllBindings)
+{
+	UMovieScene* MovieScene = Sequence->GetMovieScene();
+
+	for (int32 Index = 0; Index < MovieScene->GetPossessableCount(); ++Index)
+	{
+		FMovieScenePossessable& Possessable = MovieScene->GetPossessable(Index);
+		if (Possessable.GetParent() == Binding.BindingID)
+		{
+			FSequencerBindingProxy ChildBinding(Possessable.GetGuid(), Sequence);
+			AllBindings.Add(ChildBinding);
+			GatherDescendantBindings(ChildBinding, Sequence, AllBindings);
+		}
+	}
+}
+
 bool ExportFBXInternal(UWorld* World, UMovieSceneSequence* Sequence, const TArray<FSequencerBindingProxy>& InBindings, const TArray<UMovieSceneTrack*>& MasterTracks, UFbxExportOption* OverrideOptions, const FString& InFBXFileName, UMovieSceneSequencePlayer* Player)
 {
 	UnFbx::FFbxExporter* Exporter = UnFbx::FFbxExporter::GetInstance();
@@ -184,8 +200,17 @@ bool ExportFBXInternal(UWorld* World, UMovieSceneSequence* Sequence, const TArra
 	Exporter->SetExportOptionsOverride(OverrideOptions);
 
 	UMovieScene* MovieScene = Sequence->GetMovieScene();
+
+	TArray<FSequencerBindingProxy> AllBindings;
+	for (const FSequencerBindingProxy& Binding : InBindings)
+	{
+		AllBindings.Add(Binding);
+
+		GatherDescendantBindings(Binding, Sequence, AllBindings);
+	}
+
 	TArray<FGuid> Bindings;
-	for (const FSequencerBindingProxy& Proxy : InBindings)
+	for (const FSequencerBindingProxy& Proxy : AllBindings)
 	{
 		if (Proxy.Sequence == Sequence)
 		{
@@ -427,8 +452,16 @@ bool ImportFBXInternal(UWorld* World, UMovieSceneSequence* Sequence, const TArra
 		return false;
 	}
 
-	TMap<FGuid, FString> ObjectBindingMap;
+	TArray<FSequencerBindingProxy> AllBindings;
 	for (const FSequencerBindingProxy& Binding : InBindings)
+	{
+		AllBindings.Add(Binding);
+
+		GatherDescendantBindings(Binding, Sequence, AllBindings);
+	}
+
+	TMap<FGuid, FString> ObjectBindingMap;
+	for (const FSequencerBindingProxy& Binding : AllBindings)
 	{
 		FString Name = MovieScene->GetObjectDisplayName(Binding.BindingID).ToString();
 		ObjectBindingMap.Add(Binding.BindingID, Name);
