Sequencer: Add missing GetByteTrackEnum/SetByteTrackEnum

diff --git a/Engine/Plugins/MovieScene/SequencerScripting/Source/SequencerScripting/Private/ExtensionLibraries/MovieScenePropertyTrackExtensions.cpp b/Engine/Plugins/MovieScene/SequencerScripting/Source/SequencerScripting/Private/ExtensionLibraries/MovieScenePropertyTrackExtensions.cpp
--- a/Engine/Plugins/MovieScene/SequencerScripting/Source/SequencerScripting/Private/ExtensionLibraries/MovieScenePropertyTrackExtensions.cpp
+++ b/Engine/Plugins/MovieScene/SequencerScripting/Source/SequencerScripting/Private/ExtensionLibraries/MovieScenePropertyTrackExtensions.cpp
@@ -1,6 +1,7 @@
 // Copyright Epic Games, Inc. All Rights Reserved.
 
 #include "ExtensionLibraries/MovieScenePropertyTrackExtensions.h"
+#include "Tracks/MovieSceneByteTrack.h"
 #include "Tracks/MovieScenePropertyTrack.h"
 #include "Tracks/MovieSceneObjectPropertyTrack.h"
 
@@ -39,4 +40,26 @@ UClass* UMovieScenePropertyTrackExtensions::GetObjectPropertyClass(UMovieSceneOb
 	return Track->PropertyClass;
 }
 
+void UMovieScenePropertyTrackExtensions::SetByteTrackEnum(UMovieSceneByteTrack* Track, UEnum* InEnum)
+{
+	if (!Track)
+	{
+		FFrame::KismetExecutionMessage(TEXT("Cannot call SetByteTrackEnum on a null track"), ELogVerbosity::Error);
+		return;
+	}
+
+	Track->SetEnum(InEnum);
+}
+
+UEnum* UMovieScenePropertyTrackExtensions::GetByteTrackEnum(UMovieSceneByteTrack* Track)
+{
+	if (!Track)
+	{
+		FFrame::KismetExecutionMessage(TEXT("Cannot call GetByteTrackEnum on a null track"), ELogVerbosity::Error);
+		return nullptr;
+	}
+
+	return Track->GetEnum();
+}
+
 
diff --git a/Engine/Plugins/MovieScene/SequencerScripting/Source/SequencerScripting/Public/ExtensionLibraries/MovieScenePropertyTrackExtensions.h b/Engine/Plugins/MovieScene/SequencerScripting/Source/SequencerScripting/Public/ExtensionLibraries/MovieScenePropertyTrackExtensions.h
--- a/Engine/Plugins/MovieScene/SequencerScripting/Source/SequencerScripting/Public/ExtensionLibraries/MovieScenePropertyTrackExtensions.h
+++ b/Engine/Plugins/MovieScene/SequencerScripting/Source/SequencerScripting/Public/ExtensionLibraries/MovieScenePropertyTrackExtensions.h
@@ -8,7 +8,7 @@
 
 #include "MovieScenePropertyTrackExtensions.generated.h"
 
-
+class UMovieSceneByteTrack;
 class UMovieScenePropertyTrack;
 class UMovieSceneObjectPropertyTrack;
 
@@ -76,4 +76,22 @@ public:
 	 */
 	UFUNCTION(BlueprintCallable, Category = "Sequencer|Track", meta = (ScriptMethod))
 	static UClass* GetObjectPropertyClass(UMovieSceneObjectPropertyTrack* Track);
+
+	/**
+	 * Set this byte track's enum
+	 *
+	 * @param Track        The track to use
+	 * @param InEnum The enum to set
+	 */
+	UFUNCTION(BlueprintCallable, Category = "Sequencer|Track", meta = (ScriptMethod))
+	static void SetByteTrackEnum(UMovieSceneByteTrack* Track, UEnum* InEnum);
+
+	/**
+	 * Get this byte track's enum
+	 *
+	 * @param Track        The track to use
+	 * @return The enum for this byte track
+	 */
+	UFUNCTION(BlueprintPure, Category = "Sequencer|Track", meta = (ScriptMethod))
+	static UEnum* GetByteTrackEnum(UMovieSceneByteTrack* Track);
 };
