Sequencer Scripting: Add null checks and error messages

diff --git a/Engine/Plugins/MovieScene/SequencerScripting/Source/SequencerScripting/Private/ExtensionLibraries/MovieSceneBindingExtensions.cpp b/Engine/Plugins/MovieScene/SequencerScripting/Source/SequencerScripting/Private/ExtensionLibraries/MovieSceneBindingExtensions.cpp
--- a/Engine/Plugins/MovieScene/SequencerScripting/Source/SequencerScripting/Private/ExtensionLibraries/MovieSceneBindingExtensions.cpp
+++ b/Engine/Plugins/MovieScene/SequencerScripting/Source/SequencerScripting/Private/ExtensionLibraries/MovieSceneBindingExtensions.cpp
@@ -104,6 +104,12 @@ TArray<UMovieSceneTrack*> UMovieSceneBindingExtensions::GetTracks(const FSequenc
 
 void UMovieSceneBindingExtensions::RemoveTrack(const FSequencerBindingProxy& InBinding, UMovieSceneTrack* TrackToRemove)
 {
+	if (!TrackToRemove)
+	{
+		FFrame::KismetExecutionMessage(TEXT("Cannot call RemoveTrack on a null track"), ELogVerbosity::Error);
+		return;
+	}
+
 	UMovieScene* MovieScene = InBinding.GetMovieScene();
 	if (TrackToRemove && MovieScene)
 	{
diff --git a/Engine/Plugins/MovieScene/SequencerScripting/Source/SequencerScripting/Private/ExtensionLibraries/MovieSceneEventTrackExtensions.cpp b/Engine/Plugins/MovieScene/SequencerScripting/Source/SequencerScripting/Private/ExtensionLibraries/MovieSceneEventTrackExtensions.cpp
--- a/Engine/Plugins/MovieScene/SequencerScripting/Source/SequencerScripting/Private/ExtensionLibraries/MovieSceneEventTrackExtensions.cpp
+++ b/Engine/Plugins/MovieScene/SequencerScripting/Source/SequencerScripting/Private/ExtensionLibraries/MovieSceneEventTrackExtensions.cpp
@@ -7,6 +7,12 @@
 
 UMovieSceneEventRepeaterSection* UMovieSceneEventTrackExtensions::AddEventRepeaterSection(UMovieSceneEventTrack* Track)
 {
+	if (!Track)
+	{
+		FFrame::KismetExecutionMessage(TEXT("Cannot call AddEventRepeaterSection on a null track"), ELogVerbosity::Error);
+		return nullptr;
+	}
+
 	UMovieSceneEventRepeaterSection* NewSection = NewObject<UMovieSceneEventRepeaterSection>(Track, NAME_None, RF_Transactional);
 
 	if (NewSection)
@@ -21,6 +27,12 @@ UMovieSceneEventRepeaterSection* UMovieSceneEventTrackExtensions::AddEventRepeat
 
 UMovieSceneEventTriggerSection* UMovieSceneEventTrackExtensions::AddEventTriggerSection(UMovieSceneEventTrack* Track)
 {
+	if (!Track)
+	{
+		FFrame::KismetExecutionMessage(TEXT("Cannot call AddEventTriggerSection on a null track"), ELogVerbosity::Error);
+		return nullptr;
+	}
+
 	UMovieSceneEventTriggerSection* NewSection = NewObject<UMovieSceneEventTriggerSection>(Track, NAME_None, RF_Transactional);
 
 	if (NewSection)
diff --git a/Engine/Plugins/MovieScene/SequencerScripting/Source/SequencerScripting/Private/ExtensionLibraries/MovieSceneFolderExtensions.cpp b/Engine/Plugins/MovieScene/SequencerScripting/Source/SequencerScripting/Private/ExtensionLibraries/MovieSceneFolderExtensions.cpp
--- a/Engine/Plugins/MovieScene/SequencerScripting/Source/SequencerScripting/Private/ExtensionLibraries/MovieSceneFolderExtensions.cpp
+++ b/Engine/Plugins/MovieScene/SequencerScripting/Source/SequencerScripting/Private/ExtensionLibraries/MovieSceneFolderExtensions.cpp
@@ -7,132 +7,167 @@
 
 FName UMovieSceneFolderExtensions::GetFolderName(UMovieSceneFolder* Folder)
 {
-	if (Folder)
+	if (!Folder)
 	{
-		return Folder->GetFolderName();
+		FFrame::KismetExecutionMessage(TEXT("Cannot call GetFolderName on a null folder"), ELogVerbosity::Error);
+		return FName();
 	}
 
-	return FName();
+	return Folder->GetFolderName();
 }
 
 bool UMovieSceneFolderExtensions::SetFolderName(UMovieSceneFolder* Folder, FName InFolderName)
 {
-	if (Folder)
+	if (!Folder)
 	{
-		Folder->SetFolderName(InFolderName);
-		return true;
+		FFrame::KismetExecutionMessage(TEXT("Cannot call SetFolderName on a null folder"), ELogVerbosity::Error);
+		return false;
 	}
 
-	return false;
+	Folder->SetFolderName(InFolderName);
+	return true;
 }
 
 FColor UMovieSceneFolderExtensions::GetFolderColor(UMovieSceneFolder* Folder)
 {
-#if WITH_EDITORONLY_DATA
-	if (Folder)
+	if (!Folder)
 	{
-		return Folder->GetFolderColor();
+		FFrame::KismetExecutionMessage(TEXT("Cannot call GetFolderColor on a null folder"), ELogVerbosity::Error);
+		return FColor();
 	}
+
+#if WITH_EDITORONLY_DATA
+	return Folder->GetFolderColor();
 #endif //WITH_EDITORONLY_DATA
 	return FColor();
 }
 
 bool UMovieSceneFolderExtensions::SetFolderColor(UMovieSceneFolder* Folder, FColor InFolderColor)
 {
-#if WITH_EDITORONLY_DATA
-	if (Folder)
+	if (!Folder)
 	{
-		Folder->Modify();
-		Folder->SetFolderColor(InFolderColor);
-		return true;
+		FFrame::KismetExecutionMessage(TEXT("Cannot call SetFolderColor on a null folder"), ELogVerbosity::Error);
+		return false;
 	}
-#endif //WITH_EDITORONLY_DATA
 
-	return false;
+#if WITH_EDITORONLY_DATA
+	Folder->Modify();
+	Folder->SetFolderColor(InFolderColor);
+#endif //WITH_EDITORONLY_DATA
+	return true;
 }
 
 TArray<UMovieSceneFolder*> UMovieSceneFolderExtensions::GetChildFolders(UMovieSceneFolder* Folder)
 {
-	TArray<UMovieSceneFolder*> Result;
-
-	if (Folder)
+	if (!Folder)
 	{
-		Result = Folder->GetChildFolders();
+		FFrame::KismetExecutionMessage(TEXT("Cannot call GetChildFolders on a null folder"), ELogVerbosity::Error);
+		return TArray<UMovieSceneFolder*>();
 	}
 
-	return Result;
+	return Folder->GetChildFolders();
 }
 
 bool UMovieSceneFolderExtensions::AddChildFolder(UMovieSceneFolder* TargetFolder, UMovieSceneFolder* FolderToAdd)
 {
-	if (TargetFolder && FolderToAdd)
+	if (!TargetFolder)
 	{
-		TargetFolder->AddChildFolder(FolderToAdd);
-		return true;
+		FFrame::KismetExecutionMessage(TEXT("Cannot call AddChildFolder with a null parent folder"), ELogVerbosity::Error);
+		return false;
 	}
 
-	return false;
+	if (!FolderToAdd)
+	{
+		FFrame::KismetExecutionMessage(TEXT("Cannot call AddChildFolder with a null child folder"), ELogVerbosity::Error);
+		return false;
+	}
+
+	TargetFolder->AddChildFolder(FolderToAdd);
+	return true;
 }
 
 bool UMovieSceneFolderExtensions::RemoveChildFolder(UMovieSceneFolder* TargetFolder, UMovieSceneFolder* FolderToRemove)
 {
-	if (TargetFolder && FolderToRemove)
+	if (!TargetFolder)
 	{
-		TargetFolder->RemoveChildFolder(FolderToRemove);
-		return true;
+		FFrame::KismetExecutionMessage(TEXT("Cannot call RemoveChildFolder with a null parent folder"), ELogVerbosity::Error);
+		return false;
 	}
 
-	return false;
+	if (!FolderToRemove)
+	{
+		FFrame::KismetExecutionMessage(TEXT("Cannot call RemoveChildFolder with a null child folder"), ELogVerbosity::Error);
+		return false;
+	}
+	
+	TargetFolder->RemoveChildFolder(FolderToRemove);
+	return true;
 }
 
 TArray<UMovieSceneTrack*> UMovieSceneFolderExtensions::GetChildMasterTracks(UMovieSceneFolder* Folder)
 {
-	TArray<UMovieSceneTrack*> Result;
-
-	if (Folder)
+	if (!Folder)
 	{
-		Result = Folder->GetChildMasterTracks();
+		FFrame::KismetExecutionMessage(TEXT("Cannot call GetChildMasterTracks with a null folder"), ELogVerbosity::Error);
+		return TArray<UMovieSceneTrack*>();
 	}
 
-	return Result;
+	return Folder->GetChildMasterTracks();
 }
 
 bool UMovieSceneFolderExtensions::AddChildMasterTrack(UMovieSceneFolder* Folder, UMovieSceneTrack* InMasterTrack)
 {
-	if (Folder && InMasterTrack)
+	if (!Folder)
 	{
-		Folder->AddChildMasterTrack(InMasterTrack);
-		return true;
+		FFrame::KismetExecutionMessage(TEXT("Cannot call AddChildMasterTrack with a null folder"), ELogVerbosity::Error);
+		return false;
 	}
 
-	return false;
+	if (!InMasterTrack)
+	{
+		FFrame::KismetExecutionMessage(TEXT("Cannot call AddChildMasterTrack with a null master track"), ELogVerbosity::Error);
+		return false;
+	}
+	
+	Folder->AddChildMasterTrack(InMasterTrack);
+	return true;
 }
 
 bool UMovieSceneFolderExtensions::RemoveChildMasterTrack(UMovieSceneFolder* Folder, UMovieSceneTrack* InMasterTrack)
 {
-	if (Folder && InMasterTrack)
+	if (!Folder)
 	{
-		Folder->RemoveChildMasterTrack(InMasterTrack);
-		return true;
+		FFrame::KismetExecutionMessage(TEXT("Cannot call RemoveChildMasterTrack with a null folder"), ELogVerbosity::Error);
+		return false;
 	}
 
-	return false;
+	if (!InMasterTrack)
+	{
+		FFrame::KismetExecutionMessage(TEXT("Cannot call RemoveChildMasterTrack with a null master track"), ELogVerbosity::Error);
+		return false;
+	}
+
+	Folder->RemoveChildMasterTrack(InMasterTrack);
+	return true;
 }
 
 TArray<FSequencerBindingProxy> UMovieSceneFolderExtensions::GetChildObjectBindings(UMovieSceneFolder* Folder)
 {
+	if (!Folder)
+	{
+		FFrame::KismetExecutionMessage(TEXT("Cannot call GetChildObjectBindings with a null folder"), ELogVerbosity::Error);
+		return TArray<FSequencerBindingProxy>();
+	}
+
 	TArray<FSequencerBindingProxy> Result;
 
-	if (Folder)
-	{
-		// Attempt to get the sequence reference from the folder
-		UMovieScene* MovieScene = Cast<UMovieScene>(Folder->GetOuter());
-		UMovieSceneSequence* Sequence = Cast<UMovieSceneSequence>(MovieScene->GetOuter());
+	// Attempt to get the sequence reference from the folder
+	UMovieScene* MovieScene = Cast<UMovieScene>(Folder->GetOuter());
+	UMovieSceneSequence* Sequence = Cast<UMovieSceneSequence>(MovieScene->GetOuter());
 
-		for (FGuid ID : Folder->GetChildObjectBindings())
-		{
-			Result.Add(FSequencerBindingProxy(ID, Sequence));
-		}
+	for (FGuid ID : Folder->GetChildObjectBindings())
+	{
+		Result.Add(FSequencerBindingProxy(ID, Sequence));
 	}
 
 	return Result;
@@ -140,7 +175,13 @@ TArray<FSequencerBindingProxy> UMovieSceneFolderExtensions::GetChildObjectBindin
 
 bool UMovieSceneFolderExtensions::AddChildObjectBinding(UMovieSceneFolder* Folder, FSequencerBindingProxy InObjectBinding)
 {
-	if (Folder && InObjectBinding.BindingID.IsValid())
+	if (!Folder)
+	{
+		FFrame::KismetExecutionMessage(TEXT("Cannot call AddChildObjectBinding with a null folder"), ELogVerbosity::Error);
+		return false;
+	}
+
+	if (InObjectBinding.BindingID.IsValid())
 	{
 		Folder->AddChildObjectBinding(InObjectBinding.BindingID);
 		return true;
@@ -151,7 +192,13 @@ bool UMovieSceneFolderExtensions::AddChildObjectBinding(UMovieSceneFolder* Folde
 
 bool UMovieSceneFolderExtensions::RemoveChildObjectBinding(UMovieSceneFolder* Folder, const FSequencerBindingProxy InObjectBinding)
 {
-	if (Folder && InObjectBinding.BindingID.IsValid())
+	if (!Folder)
+	{
+		FFrame::KismetExecutionMessage(TEXT("Cannot call RemoveChildObjectBinding with a null folder"), ELogVerbosity::Error);
+		return false;
+	}
+
+	if (InObjectBinding.BindingID.IsValid())
 	{
 		Folder->RemoveChildObjectBinding(InObjectBinding.BindingID);
 		return true;
diff --git a/Engine/Plugins/MovieScene/SequencerScripting/Source/SequencerScripting/Private/ExtensionLibraries/MovieScenePropertyTrackExtensions.cpp b/Engine/Plugins/MovieScene/SequencerScripting/Source/SequencerScripting/Private/ExtensionLibraries/MovieScenePropertyTrackExtensions.cpp
--- a/Engine/Plugins/MovieScene/SequencerScripting/Source/SequencerScripting/Private/ExtensionLibraries/MovieScenePropertyTrackExtensions.cpp
+++ b/Engine/Plugins/MovieScene/SequencerScripting/Source/SequencerScripting/Private/ExtensionLibraries/MovieScenePropertyTrackExtensions.cpp
@@ -7,22 +7,46 @@
 
 void UMovieScenePropertyTrackExtensions::SetPropertyNameAndPath(UMovieScenePropertyTrack* Track, const FName& InPropertyName, const FString& InPropertyPath)
 {
+	if (!Track)
+	{
+		FFrame::KismetExecutionMessage(TEXT("Cannot call SetPropertyNameAndPath on a null track"), ELogVerbosity::Error);
+		return;
+	}
+
 	Track->Modify();
 	Track->SetPropertyNameAndPath(InPropertyName, InPropertyPath);
 }
 
 FName UMovieScenePropertyTrackExtensions::GetPropertyName(UMovieScenePropertyTrack* Track)
 {
+	if (!Track)
+	{
+		FFrame::KismetExecutionMessage(TEXT("Cannot call GetPropertyName on a null track"), ELogVerbosity::Error);
+		return NAME_None;
+	}
+
 	return Track->GetPropertyName();
 }
 
 FString UMovieScenePropertyTrackExtensions::GetPropertyPath(UMovieScenePropertyTrack* Track)
 {
+	if (!Track)
+	{
+		FFrame::KismetExecutionMessage(TEXT("Cannot call GetPropertyPath on a null track"), ELogVerbosity::Error);
+		return FString();
+	}
+
 	return Track->GetPropertyPath().ToString();
 }
 
 FName UMovieScenePropertyTrackExtensions::GetUniqueTrackName(UMovieScenePropertyTrack* Track)
 {
+	if (!Track)
+	{
+		FFrame::KismetExecutionMessage(TEXT("Cannot call GetUniqueTrackName on a null track"), ELogVerbosity::Error);
+		return NAME_None;
+	}
+
 #if WITH_EDITORONLY_DATA
 	return Track->UniqueTrackName;
 #else
@@ -32,11 +56,23 @@ FName UMovieScenePropertyTrackExtensions::GetUniqueTrackName(UMovieSceneProperty
 
 void UMovieScenePropertyTrackExtensions::SetObjectPropertyClass(UMovieSceneObjectPropertyTrack* Track, UClass* PropertyClass)
 {
+	if (!Track)
+	{
+		FFrame::KismetExecutionMessage(TEXT("Cannot call SetObjectPropertyClass on a null track"), ELogVerbosity::Error);
+		return;
+	}
+
 	Track->PropertyClass = PropertyClass;
 }
 
 UClass* UMovieScenePropertyTrackExtensions::GetObjectPropertyClass(UMovieSceneObjectPropertyTrack* Track)
 {
+	if (!Track)
+	{
+		FFrame::KismetExecutionMessage(TEXT("Cannot call AddEventRepeaterSection on a null track"), ELogVerbosity::Error);
+		return nullptr;
+	}
+
 	return Track->PropertyClass;
 }
 
diff --git a/Engine/Plugins/MovieScene/SequencerScripting/Source/SequencerScripting/Private/ExtensionLibraries/MovieSceneSectionExtensions.cpp b/Engine/Plugins/MovieScene/SequencerScripting/Source/SequencerScripting/Private/ExtensionLibraries/MovieSceneSectionExtensions.cpp
--- a/Engine/Plugins/MovieScene/SequencerScripting/Source/SequencerScripting/Private/ExtensionLibraries/MovieSceneSectionExtensions.cpp
+++ b/Engine/Plugins/MovieScene/SequencerScripting/Source/SequencerScripting/Private/ExtensionLibraries/MovieSceneSectionExtensions.cpp
@@ -22,17 +22,35 @@
 
 FSequencerScriptingRange UMovieSceneSectionExtensions::GetRange(UMovieSceneSection* Section)
 {
+	if (!Section)
+	{
+		FFrame::KismetExecutionMessage(TEXT("Cannot call GetRange on a null section"), ELogVerbosity::Error);
+		return FSequencerScriptingRange();
+	}
+
 	UMovieScene* MovieScene = Section->GetTypedOuter<UMovieScene>();
 	return FSequencerScriptingRange::FromNative(Section->GetRange(), MovieScene->GetTickResolution());
 }
 
 bool UMovieSceneSectionExtensions::HasStartFrame(UMovieSceneSection* Section)
 {
+	if (!Section)
+	{
+		FFrame::KismetExecutionMessage(TEXT("Cannot call HasStartFrame on a null section"), ELogVerbosity::Error);
+		return false;
+	}
+
 	return Section->HasStartFrame();
 }
 
 int32 UMovieSceneSectionExtensions::GetStartFrame(UMovieSceneSection* Section)
 {
+	if (!Section)
+	{
+		FFrame::KismetExecutionMessage(TEXT("Cannot call GetStartFrame on a null section"), ELogVerbosity::Error);
+		return -1;
+	}
+
 	if (!Section->HasStartFrame())
 	{
 		FFrame::KismetExecutionMessage(TEXT("Section does not have a start frame"), ELogVerbosity::Error);
@@ -53,6 +71,12 @@ int32 UMovieSceneSectionExtensions::GetStartFrame(UMovieSceneSection* Section)
 
 float UMovieSceneSectionExtensions::GetStartFrameSeconds(UMovieSceneSection* Section)
 {
+	if (!Section)
+	{
+		FFrame::KismetExecutionMessage(TEXT("Cannot call GetStartFrameSeconds on a null section"), ELogVerbosity::Error);
+		return -1.f;
+	}
+
 	if (!Section->HasStartFrame())
 	{
 		FFrame::KismetExecutionMessage(TEXT("Section does not have a start frame"), ELogVerbosity::Error);
@@ -71,12 +95,24 @@ float UMovieSceneSectionExtensions::GetStartFrameSeconds(UMovieSceneSection* Sec
 
 bool UMovieSceneSectionExtensions::HasEndFrame(UMovieSceneSection* Section)
 {
+	if (!Section)
+	{
+		FFrame::KismetExecutionMessage(TEXT("Cannot call HasEndFrame on a null section"), ELogVerbosity::Error);
+		return false;
+	}
+
 	return Section->HasEndFrame();
 }
 
 
 int32 UMovieSceneSectionExtensions::GetEndFrame(UMovieSceneSection* Section)
 {
+	if (!Section)
+	{
+		FFrame::KismetExecutionMessage(TEXT("Cannot call GetEndFrame on a null section"), ELogVerbosity::Error);
+		return -1;
+	}
+
 	if (!Section->HasEndFrame())
 	{
 		FFrame::KismetExecutionMessage(TEXT("Section does not have an end frame"), ELogVerbosity::Error);
@@ -97,6 +133,12 @@ int32 UMovieSceneSectionExtensions::GetEndFrame(UMovieSceneSection* Section)
 
 float UMovieSceneSectionExtensions::GetEndFrameSeconds(UMovieSceneSection* Section)
 {
+	if (!Section)
+	{
+		FFrame::KismetExecutionMessage(TEXT("Cannot call GetEndFrameSeconds on a null section"), ELogVerbosity::Error);
+		return -1.f;
+	}
+
 	if (!Section->HasEndFrame())
 	{
 		FFrame::KismetExecutionMessage(TEXT("Section does not have an end frame"), ELogVerbosity::Error);
@@ -115,6 +157,12 @@ float UMovieSceneSectionExtensions::GetEndFrameSeconds(UMovieSceneSection* Secti
 
 void UMovieSceneSectionExtensions::SetRange(UMovieSceneSection* Section, int32 StartFrame, int32 EndFrame)
 {
+	if (!Section)
+	{
+		FFrame::KismetExecutionMessage(TEXT("Cannot call SetRange on a null section"), ELogVerbosity::Error);
+		return;
+	}
+
 	UMovieScene* MovieScene = Section->GetTypedOuter<UMovieScene>();
 	if (MovieScene)
 	{
@@ -137,6 +185,12 @@ void UMovieSceneSectionExtensions::SetRange(UMovieSceneSection* Section, int32 S
 
 void UMovieSceneSectionExtensions::SetRangeSeconds(UMovieSceneSection* Section, float StartTime, float EndTime)
 {
+	if (!Section)
+	{
+		FFrame::KismetExecutionMessage(TEXT("Cannot call SetRangeSeconds on a null section"), ELogVerbosity::Error);
+		return;
+	}
+
 	UMovieScene* MovieScene = Section->GetTypedOuter<UMovieScene>();
 	if (MovieScene)
 	{
@@ -159,6 +213,12 @@ void UMovieSceneSectionExtensions::SetRangeSeconds(UMovieSceneSection* Section,
 
 void UMovieSceneSectionExtensions::SetStartFrame(UMovieSceneSection* Section, int32 StartFrame)
 {
+	if (!Section)
+	{
+		FFrame::KismetExecutionMessage(TEXT("Cannot call SetStartFrame on a null section"), ELogVerbosity::Error);
+		return;
+	}
+
 	UMovieScene* MovieScene = Section->GetTypedOuter<UMovieScene>();
 	if (MovieScene)
 	{
@@ -170,6 +230,12 @@ void UMovieSceneSectionExtensions::SetStartFrame(UMovieSceneSection* Section, in
 
 void UMovieSceneSectionExtensions::SetStartFrameSeconds(UMovieSceneSection* Section, float StartTime)
 {
+	if (!Section)
+	{
+		FFrame::KismetExecutionMessage(TEXT("Cannot call SetStartFrameSeconds on a null section"), ELogVerbosity::Error);
+		return;
+	}
+
 	UMovieScene* MovieScene = Section->GetTypedOuter<UMovieScene>();
 	if (MovieScene)
 	{
@@ -179,6 +245,12 @@ void UMovieSceneSectionExtensions::SetStartFrameSeconds(UMovieSceneSection* Sect
 
 void UMovieSceneSectionExtensions::SetStartFrameBounded(UMovieSceneSection* Section, bool bIsBounded)
 {
+	if (!Section)
+	{
+		FFrame::KismetExecutionMessage(TEXT("Cannot call SetStartFrameBounded on a null section"), ELogVerbosity::Error);
+		return;
+	}
+
 	UMovieScene* MovieScene = Section->GetTypedOuter<UMovieScene>();
 	if (MovieScene)
 	{
@@ -201,6 +273,12 @@ void UMovieSceneSectionExtensions::SetStartFrameBounded(UMovieSceneSection* Sect
 
 void UMovieSceneSectionExtensions::SetEndFrame(UMovieSceneSection* Section, int32 EndFrame)
 {
+	if (!Section)
+	{
+		FFrame::KismetExecutionMessage(TEXT("Cannot call SetEndFrame on a null section"), ELogVerbosity::Error);
+		return;
+	}
+
 	UMovieScene* MovieScene = Section->GetTypedOuter<UMovieScene>();
 	if (MovieScene)
 	{
@@ -212,6 +290,12 @@ void UMovieSceneSectionExtensions::SetEndFrame(UMovieSceneSection* Section, int3
 
 void UMovieSceneSectionExtensions::SetEndFrameSeconds(UMovieSceneSection* Section, float EndTime)
 {
+	if (!Section)
+	{
+		FFrame::KismetExecutionMessage(TEXT("Cannot call SetEndFrameSeconds on a null section"), ELogVerbosity::Error);
+		return;
+	}
+
 	UMovieScene* MovieScene = Section->GetTypedOuter<UMovieScene>();
 	if (MovieScene)
 	{
@@ -221,6 +305,12 @@ void UMovieSceneSectionExtensions::SetEndFrameSeconds(UMovieSceneSection* Sectio
 
 void UMovieSceneSectionExtensions::SetEndFrameBounded(UMovieSceneSection* Section, bool bIsBounded)
 {
+	if (!Section)
+	{
+		FFrame::KismetExecutionMessage(TEXT("Cannot call SetEndFrameBounded on a null section"), ELogVerbosity::Error);
+		return;
+	}
+
 	UMovieScene* MovieScene = Section->GetTypedOuter<UMovieScene>();
 	if (MovieScene)
 	{
diff --git a/Engine/Plugins/MovieScene/SequencerScripting/Source/SequencerScripting/Private/ExtensionLibraries/MovieSceneSequenceExtensions.cpp b/Engine/Plugins/MovieScene/SequencerScripting/Source/SequencerScripting/Private/ExtensionLibraries/MovieSceneSequenceExtensions.cpp
--- a/Engine/Plugins/MovieScene/SequencerScripting/Source/SequencerScripting/Private/ExtensionLibraries/MovieSceneSequenceExtensions.cpp
+++ b/Engine/Plugins/MovieScene/SequencerScripting/Source/SequencerScripting/Private/ExtensionLibraries/MovieSceneSequenceExtensions.cpp
@@ -27,11 +27,23 @@ TArray<UMovieSceneTrack*> UMovieSceneSequenceExtensions::FilterTracks(TArrayView
 
 UMovieScene* UMovieSceneSequenceExtensions::GetMovieScene(UMovieSceneSequence* Sequence)
 {
-	return Sequence ? Sequence->GetMovieScene() : nullptr;
+	if (!Sequence)
+	{
+		FFrame::KismetExecutionMessage(TEXT("Cannot call GetMovieScene on a null sequence"), ELogVerbosity::Error);
+		return nullptr;
+	}
+
+	return Sequence->GetMovieScene();
 }
 
 TArray<UMovieSceneTrack*> UMovieSceneSequenceExtensions::GetMasterTracks(UMovieSceneSequence* Sequence)
 {
+	if (!Sequence)
+	{
+		FFrame::KismetExecutionMessage(TEXT("Cannot call GetMasterTracks on a null sequence"), ELogVerbosity::Error);
+		return TArray<UMovieSceneTrack*>();
+	}
+
 	TArray<UMovieSceneTrack*> Tracks;
 
 	UMovieScene* MovieScene = GetMovieScene(Sequence);
@@ -50,6 +62,12 @@ TArray<UMovieSceneTrack*> UMovieSceneSequenceExtensions::GetMasterTracks(UMovieS
 
 TArray<UMovieSceneTrack*> UMovieSceneSequenceExtensions::FindMasterTracksByType(UMovieSceneSequence* Sequence, TSubclassOf<UMovieSceneTrack> TrackType)
 {
+	if (!Sequence)
+	{
+		FFrame::KismetExecutionMessage(TEXT("Cannot call FindMasterTracksByType on a null sequence"), ELogVerbosity::Error);
+		return TArray<UMovieSceneTrack*>();
+	}
+
 	UMovieScene* MovieScene   = GetMovieScene(Sequence);
 	UClass*      DesiredClass = TrackType.Get();
 
@@ -73,6 +91,12 @@ TArray<UMovieSceneTrack*> UMovieSceneSequenceExtensions::FindMasterTracksByType(
 
 TArray<UMovieSceneTrack*> UMovieSceneSequenceExtensions::FindMasterTracksByExactType(UMovieSceneSequence* Sequence, TSubclassOf<UMovieSceneTrack> TrackType)
 {
+	if (!Sequence)
+	{
+		FFrame::KismetExecutionMessage(TEXT("Cannot call FindMasterTracksByExactType on a null sequence"), ELogVerbosity::Error);
+		return TArray<UMovieSceneTrack*>();
+	}
+
 	UMovieScene* MovieScene   = GetMovieScene(Sequence);
 	UClass*      DesiredClass = TrackType.Get();
 
@@ -96,6 +120,12 @@ TArray<UMovieSceneTrack*> UMovieSceneSequenceExtensions::FindMasterTracksByExact
 
 UMovieSceneTrack* UMovieSceneSequenceExtensions::AddMasterTrack(UMovieSceneSequence* Sequence, TSubclassOf<UMovieSceneTrack> TrackType)
 {
+	if (!Sequence)
+	{
+		FFrame::KismetExecutionMessage(TEXT("Cannot call AddMasterTrack on a null sequence"), ELogVerbosity::Error);
+		return nullptr;
+	}
+
 	// @todo: sequencer-python: master track type compatibility with sequence. Currently that's really only loosely defined by track editors, which is not sufficient here.
 	UMovieScene* MovieScene = GetMovieScene(Sequence);
 	if (MovieScene)
@@ -115,6 +145,12 @@ UMovieSceneTrack* UMovieSceneSequenceExtensions::AddMasterTrack(UMovieSceneSeque
 
 bool UMovieSceneSequenceExtensions::RemoveMasterTrack(UMovieSceneSequence* Sequence, UMovieSceneTrack* MasterTrack)
 {
+	if (!Sequence)
+	{
+		FFrame::KismetExecutionMessage(TEXT("Cannot call RemoveMasterTrack on a null sequence"), ELogVerbosity::Error);
+		return false;
+	}
+
 	UMovieScene* MovieScene = GetMovieScene(Sequence);
 	if (MovieScene)
 	{
@@ -135,12 +171,24 @@ bool UMovieSceneSequenceExtensions::RemoveMasterTrack(UMovieSceneSequence* Seque
 
 FFrameRate UMovieSceneSequenceExtensions::GetDisplayRate(UMovieSceneSequence* Sequence)
 {
+	if (!Sequence)
+	{
+		FFrame::KismetExecutionMessage(TEXT("Cannot call GetDisplayRate on a null sequence"), ELogVerbosity::Error);
+		return FFrameRate();
+	}
+
 	UMovieScene* MovieScene = GetMovieScene(Sequence);
 	return MovieScene ? MovieScene->GetDisplayRate() : FFrameRate();
 }
 
 void UMovieSceneSequenceExtensions::SetDisplayRate(UMovieSceneSequence* Sequence, FFrameRate DisplayRate)
 {
+	if (!Sequence)
+	{
+		FFrame::KismetExecutionMessage(TEXT("Cannot call SetDisplayRate on a null sequence"), ELogVerbosity::Error);
+		return;
+	}
+
 	UMovieScene* MovieScene = GetMovieScene(Sequence);
 	if (MovieScene)
 	{
@@ -152,12 +200,24 @@ void UMovieSceneSequenceExtensions::SetDisplayRate(UMovieSceneSequence* Sequence
 
 FFrameRate UMovieSceneSequenceExtensions::GetTickResolution(UMovieSceneSequence* Sequence)
 {
+	if (!Sequence)
+	{
+		FFrame::KismetExecutionMessage(TEXT("Cannot call GetTickResolution on a null sequence"), ELogVerbosity::Error);
+		return FFrameRate();
+	}
+
 	UMovieScene* MovieScene = GetMovieScene(Sequence);
 	return MovieScene ? MovieScene->GetTickResolution() : FFrameRate();
 }
 
 void UMovieSceneSequenceExtensions::SetTickResolution(UMovieSceneSequence* Sequence, FFrameRate TickResolution)
 {
+	if (!Sequence)
+	{
+		FFrame::KismetExecutionMessage(TEXT("Cannot call SetTickResolution on a null sequence"), ELogVerbosity::Error);
+		return;
+	}
+
 	UMovieScene* MovieScene = GetMovieScene(Sequence);
 	if (MovieScene)
 	{
@@ -169,6 +229,12 @@ void UMovieSceneSequenceExtensions::SetTickResolution(UMovieSceneSequence* Seque
 
 void UMovieSceneSequenceExtensions::SetTickResolutionDirectly(UMovieSceneSequence* Sequence, FFrameRate TickResolution)
 {
+	if (!Sequence)
+	{
+		FFrame::KismetExecutionMessage(TEXT("Cannot call SetTickResolutionDirectly on a null sequence"), ELogVerbosity::Error);
+		return;
+	}
+
 	UMovieScene* MovieScene = GetMovieScene(Sequence);
 	if (MovieScene)	
 	{
@@ -180,18 +246,36 @@ void UMovieSceneSequenceExtensions::SetTickResolutionDirectly(UMovieSceneSequenc
 
 FSequencerScriptingRange UMovieSceneSequenceExtensions::MakeRange(UMovieSceneSequence* Sequence, int32 StartFrame, int32 Duration)
 {
+	if (!Sequence)
+	{
+		FFrame::KismetExecutionMessage(TEXT("Cannot call MakeRange on a null sequence"), ELogVerbosity::Error);
+		return FSequencerScriptingRange();
+	}
+
 	FFrameRate FrameRate = GetDisplayRate(Sequence);
 	return FSequencerScriptingRange::FromNative(TRange<FFrameNumber>(StartFrame, StartFrame+Duration), FrameRate, FrameRate);
 }
 
 FSequencerScriptingRange UMovieSceneSequenceExtensions::MakeRangeSeconds(UMovieSceneSequence* Sequence, float StartTime, float Duration)
 {
+	if (!Sequence)
+	{
+		FFrame::KismetExecutionMessage(TEXT("Cannot call MakeRangeSeconds on a null sequence"), ELogVerbosity::Error);
+		return FSequencerScriptingRange();
+	}
+
 	FFrameRate FrameRate = GetDisplayRate(Sequence);
 	return FSequencerScriptingRange::FromNative(TRange<FFrameNumber>((StartTime*FrameRate).FloorToFrame(), ((StartTime+Duration) * FrameRate).CeilToFrame()), FrameRate, FrameRate);
 }
 
 FSequencerScriptingRange UMovieSceneSequenceExtensions::GetPlaybackRange(UMovieSceneSequence* Sequence)
 {
+	if (!Sequence)
+	{
+		FFrame::KismetExecutionMessage(TEXT("Cannot call GetPlaybackRange on a null sequence"), ELogVerbosity::Error);
+		return FSequencerScriptingRange();
+	}
+
 	UMovieScene* MovieScene = GetMovieScene(Sequence);
 	if (MovieScene)
 	{
@@ -205,6 +289,12 @@ FSequencerScriptingRange UMovieSceneSequenceExtensions::GetPlaybackRange(UMovieS
 
 int32 UMovieSceneSequenceExtensions::GetPlaybackStart(UMovieSceneSequence* Sequence)
 {
+	if (!Sequence)
+	{
+		FFrame::KismetExecutionMessage(TEXT("Cannot call GetPlaybackStart on a null sequence"), ELogVerbosity::Error);
+		return -1;
+	}
+
 	UMovieScene* MovieScene = GetMovieScene(Sequence);
 	if (MovieScene)
 	{
@@ -219,6 +309,12 @@ int32 UMovieSceneSequenceExtensions::GetPlaybackStart(UMovieSceneSequence* Seque
 
 float UMovieSceneSequenceExtensions::GetPlaybackStartSeconds(UMovieSceneSequence* Sequence)
 {
+	if (!Sequence)
+	{
+		FFrame::KismetExecutionMessage(TEXT("Cannot call GetPlaybackStartSeconds on a null sequence"), ELogVerbosity::Error);
+		return -1.f;
+	}
+
 	UMovieScene* MovieScene = GetMovieScene(Sequence);
 	if (MovieScene)
 	{
@@ -231,6 +327,12 @@ float UMovieSceneSequenceExtensions::GetPlaybackStartSeconds(UMovieSceneSequence
 
 int32 UMovieSceneSequenceExtensions::GetPlaybackEnd(UMovieSceneSequence* Sequence)
 {
+	if (!Sequence)
+	{
+		FFrame::KismetExecutionMessage(TEXT("Cannot call GetPlaybackEnd on a null sequence"), ELogVerbosity::Error);
+		return -1;
+	}
+
 	UMovieScene* MovieScene = GetMovieScene(Sequence);
 	if (MovieScene)
 	{
@@ -245,6 +347,12 @@ int32 UMovieSceneSequenceExtensions::GetPlaybackEnd(UMovieSceneSequence* Sequenc
 
 float UMovieSceneSequenceExtensions::GetPlaybackEndSeconds(UMovieSceneSequence* Sequence)
 {
+	if (!Sequence)
+	{
+		FFrame::KismetExecutionMessage(TEXT("Cannot call GetPlaybackEndSeconds on a null sequence"), ELogVerbosity::Error);
+		return -1.f;
+	}
+
 	UMovieScene* MovieScene = GetMovieScene(Sequence);
 	if (MovieScene)
 	{
@@ -257,6 +365,12 @@ float UMovieSceneSequenceExtensions::GetPlaybackEndSeconds(UMovieSceneSequence*
 
 void UMovieSceneSequenceExtensions::SetPlaybackStart(UMovieSceneSequence* Sequence, int32 StartFrame)
 {
+	if (!Sequence)
+	{
+		FFrame::KismetExecutionMessage(TEXT("Cannot call SetPlaybackStart on a null sequence"), ELogVerbosity::Error);
+		return;
+	}
+
 	UMovieScene* MovieScene = GetMovieScene(Sequence);
 	if (MovieScene)
 	{
@@ -271,6 +385,12 @@ void UMovieSceneSequenceExtensions::SetPlaybackStart(UMovieSceneSequence* Sequen
 
 void UMovieSceneSequenceExtensions::SetPlaybackStartSeconds(UMovieSceneSequence* Sequence, float StartTime)
 {
+	if (!Sequence)
+	{
+		FFrame::KismetExecutionMessage(TEXT("Cannot call SetPlaybackStartSeconds on a null sequence"), ELogVerbosity::Error);
+		return;
+	}
+
 	UMovieScene* MovieScene = GetMovieScene(Sequence);
 	if (MovieScene)
 	{
@@ -283,6 +403,12 @@ void UMovieSceneSequenceExtensions::SetPlaybackStartSeconds(UMovieSceneSequence*
 
 void UMovieSceneSequenceExtensions::SetPlaybackEnd(UMovieSceneSequence* Sequence, int32 EndFrame)
 {
+	if (!Sequence)
+	{
+		FFrame::KismetExecutionMessage(TEXT("Cannot call SetPlaybackEnd on a null sequence"), ELogVerbosity::Error);
+		return;
+	}
+
 	UMovieScene* MovieScene = GetMovieScene(Sequence);
 	if (MovieScene)
 	{
@@ -297,6 +423,12 @@ void UMovieSceneSequenceExtensions::SetPlaybackEnd(UMovieSceneSequence* Sequence
 
 void UMovieSceneSequenceExtensions::SetPlaybackEndSeconds(UMovieSceneSequence* Sequence, float EndTime)
 {
+	if (!Sequence)
+	{
+		FFrame::KismetExecutionMessage(TEXT("Cannot call SetPlaybackEndSeconds on a null sequence"), ELogVerbosity::Error);
+		return;
+	}
+
 	UMovieScene* MovieScene = GetMovieScene(Sequence);
 	if (MovieScene)
 	{
@@ -309,6 +441,12 @@ void UMovieSceneSequenceExtensions::SetPlaybackEndSeconds(UMovieSceneSequence* S
 
 void UMovieSceneSequenceExtensions::SetViewRangeStart(UMovieSceneSequence* Sequence, float StartTimeInSeconds)
 {
+	if (!Sequence)
+	{
+		FFrame::KismetExecutionMessage(TEXT("Cannot call SetViewRangeStart on a null sequence"), ELogVerbosity::Error);
+		return;
+	}
+
 	UMovieScene* MovieScene = GetMovieScene(Sequence);
 	if (MovieScene)
 	{
@@ -320,6 +458,12 @@ void UMovieSceneSequenceExtensions::SetViewRangeStart(UMovieSceneSequence* Seque
 
 float UMovieSceneSequenceExtensions::GetViewRangeStart(UMovieSceneSequence* Sequence)
 {
+	if (!Sequence)
+	{
+		FFrame::KismetExecutionMessage(TEXT("Cannot call GetViewRangeStart on a null sequence"), ELogVerbosity::Error);
+		return 0.f;
+	}
+
 	UMovieScene* MovieScene = GetMovieScene(Sequence);
 	if (MovieScene)
 	{
@@ -332,6 +476,12 @@ float UMovieSceneSequenceExtensions::GetViewRangeStart(UMovieSceneSequence* Sequ
 
 void UMovieSceneSequenceExtensions::SetViewRangeEnd(UMovieSceneSequence* Sequence, float EndTimeInSeconds)
 {
+	if (!Sequence)
+	{
+		FFrame::KismetExecutionMessage(TEXT("Cannot call SetViewRangeEnd on a null sequence"), ELogVerbosity::Error);
+		return;
+	}
+
 	UMovieScene* MovieScene = GetMovieScene(Sequence);
 	if (MovieScene)
 	{
@@ -343,6 +493,12 @@ void UMovieSceneSequenceExtensions::SetViewRangeEnd(UMovieSceneSequence* Sequenc
 
 float UMovieSceneSequenceExtensions::GetViewRangeEnd(UMovieSceneSequence* Sequence)
 {
+	if (!Sequence)
+	{
+		FFrame::KismetExecutionMessage(TEXT("Cannot call GetViewRangeEnd on a null sequence"), ELogVerbosity::Error);
+		return 0.f;
+	}
+
 	UMovieScene* MovieScene = GetMovieScene(Sequence);
 	if (MovieScene)
 	{
@@ -355,6 +511,12 @@ float UMovieSceneSequenceExtensions::GetViewRangeEnd(UMovieSceneSequence* Sequen
 
 void UMovieSceneSequenceExtensions::SetWorkRangeStart(UMovieSceneSequence* Sequence, float StartTimeInSeconds)
 {
+	if (!Sequence)
+	{
+		FFrame::KismetExecutionMessage(TEXT("Cannot call SetWorkRangeStart on a null sequence"), ELogVerbosity::Error);
+		return;
+	}
+
 	UMovieScene* MovieScene = GetMovieScene(Sequence);
 	if (MovieScene)
 	{
@@ -366,6 +528,12 @@ void UMovieSceneSequenceExtensions::SetWorkRangeStart(UMovieSceneSequence* Seque
 
 float UMovieSceneSequenceExtensions::GetWorkRangeStart(UMovieSceneSequence* Sequence)
 {
+	if (!Sequence)
+	{
+		FFrame::KismetExecutionMessage(TEXT("Cannot call GetWorkRangeStart on a null sequence"), ELogVerbosity::Error);
+		return 0.f;
+	}
+
 	UMovieScene* MovieScene = GetMovieScene(Sequence);
 	if (MovieScene)
 	{
@@ -378,6 +546,12 @@ float UMovieSceneSequenceExtensions::GetWorkRangeStart(UMovieSceneSequence* Sequ
 
 void UMovieSceneSequenceExtensions::SetWorkRangeEnd(UMovieSceneSequence* Sequence, float EndTimeInSeconds)
 {
+	if (!Sequence)
+	{
+		FFrame::KismetExecutionMessage(TEXT("Cannot call SetWorkRangeEnd on a null sequence"), ELogVerbosity::Error);
+		return;
+	}
+
 	UMovieScene* MovieScene = GetMovieScene(Sequence);
 	if (MovieScene)
 	{
@@ -389,6 +563,12 @@ void UMovieSceneSequenceExtensions::SetWorkRangeEnd(UMovieSceneSequence* Sequenc
 
 float UMovieSceneSequenceExtensions::GetWorkRangeEnd(UMovieSceneSequence* Sequence)
 {
+	if (!Sequence)
+	{
+		FFrame::KismetExecutionMessage(TEXT("Cannot call GetWorkRangeEnd on a null sequence"), ELogVerbosity::Error);
+		return 0.f;
+	}
+
 	UMovieScene* MovieScene = GetMovieScene(Sequence);
 	if (MovieScene)
 	{
@@ -401,6 +581,12 @@ float UMovieSceneSequenceExtensions::GetWorkRangeEnd(UMovieSceneSequence* Sequen
 
 void UMovieSceneSequenceExtensions::SetEvaluationType(UMovieSceneSequence* Sequence, EMovieSceneEvaluationType InEvaluationType)
 {
+	if (!Sequence)
+	{
+		FFrame::KismetExecutionMessage(TEXT("Cannot call SetEvaluationType on a null sequence"), ELogVerbosity::Error);
+		return;
+	}
+
 	UMovieScene* MovieScene = GetMovieScene(Sequence);
 	if (MovieScene)
 	{
@@ -410,6 +596,12 @@ void UMovieSceneSequenceExtensions::SetEvaluationType(UMovieSceneSequence* Seque
 
 EMovieSceneEvaluationType UMovieSceneSequenceExtensions::GetEvaluationType(UMovieSceneSequence* Sequence)
 {
+	if (!Sequence)
+	{
+		FFrame::KismetExecutionMessage(TEXT("Cannot call GetEvaluationType on a null sequence"), ELogVerbosity::Error);
+		return EMovieSceneEvaluationType::WithSubFrames;
+	}
+
 	UMovieScene* MovieScene = GetMovieScene(Sequence);
 	if (MovieScene)
 	{
@@ -421,6 +613,12 @@ EMovieSceneEvaluationType UMovieSceneSequenceExtensions::GetEvaluationType(UMovi
 
 void UMovieSceneSequenceExtensions::SetClockSource(UMovieSceneSequence* Sequence, EUpdateClockSource InClockSource)
 {
+	if (!Sequence)
+	{
+		FFrame::KismetExecutionMessage(TEXT("Cannot call SetClockSource on a null sequence"), ELogVerbosity::Error);
+		return;
+	}
+
 	UMovieScene* MovieScene = GetMovieScene(Sequence);
 	if (MovieScene)
 	{
@@ -430,6 +628,12 @@ void UMovieSceneSequenceExtensions::SetClockSource(UMovieSceneSequence* Sequence
 
 EUpdateClockSource UMovieSceneSequenceExtensions::GetClockSource(UMovieSceneSequence* Sequence)
 {
+	if (!Sequence)
+	{
+		FFrame::KismetExecutionMessage(TEXT("Cannot call GetClockSource on a null sequence"), ELogVerbosity::Error);
+		return EUpdateClockSource::Tick;
+	}
+
 	UMovieScene* MovieScene = GetMovieScene(Sequence);
 	if (MovieScene)
 	{
@@ -441,6 +645,12 @@ EUpdateClockSource UMovieSceneSequenceExtensions::GetClockSource(UMovieSceneSequ
 
 FTimecode UMovieSceneSequenceExtensions::GetTimecodeSource(UMovieSceneSequence* Sequence)
 {
+	if (!Sequence)
+	{
+		FFrame::KismetExecutionMessage(TEXT("Cannot call GetTimecodeSource on a null sequence"), ELogVerbosity::Error);
+		return FTimecode();
+	}
+
 	UMovieScene* MovieScene = GetMovieScene(Sequence);
 	if (!MovieScene)
 	{
@@ -452,6 +662,12 @@ FTimecode UMovieSceneSequenceExtensions::GetTimecodeSource(UMovieSceneSequence*
 
 FSequencerBindingProxy UMovieSceneSequenceExtensions::FindBindingByName(UMovieSceneSequence* Sequence, FString Name)
 {
+	if (!Sequence)
+	{
+		FFrame::KismetExecutionMessage(TEXT("Cannot call FindBindingByName on a null sequence"), ELogVerbosity::Error);
+		return FSequencerBindingProxy();
+	}
+
 	UMovieScene* MovieScene = GetMovieScene(Sequence);
 	if (MovieScene)
 	{
@@ -466,6 +682,12 @@ FSequencerBindingProxy UMovieSceneSequenceExtensions::FindBindingByName(UMovieSc
 
 FSequencerBindingProxy UMovieSceneSequenceExtensions::FindBindingById(UMovieSceneSequence* Sequence, FGuid BindingId)
 {
+	if (!Sequence)
+	{
+		FFrame::KismetExecutionMessage(TEXT("Cannot call FindBindingById on a null sequence"), ELogVerbosity::Error);
+		return FSequencerBindingProxy();
+	}
+
 	UMovieScene* MovieScene = GetMovieScene(Sequence);
 	if (MovieScene)
 	{
@@ -480,6 +702,12 @@ FSequencerBindingProxy UMovieSceneSequenceExtensions::FindBindingById(UMovieScen
 
 TArray<FSequencerBindingProxy> UMovieSceneSequenceExtensions::GetBindings(UMovieSceneSequence* Sequence)
 {
+	if (!Sequence)
+	{
+		FFrame::KismetExecutionMessage(TEXT("Cannot call GetBindings on a null sequence"), ELogVerbosity::Error);
+		return TArray<FSequencerBindingProxy>();
+	}
+
 	TArray<FSequencerBindingProxy> AllBindings;
 
 	UMovieScene* MovieScene = GetMovieScene(Sequence);
@@ -496,6 +724,12 @@ TArray<FSequencerBindingProxy> UMovieSceneSequenceExtensions::GetBindings(UMovie
 
 TArray<FSequencerBindingProxy> UMovieSceneSequenceExtensions::GetSpawnables(UMovieSceneSequence* Sequence)
 {
+	if (!Sequence)
+	{
+		FFrame::KismetExecutionMessage(TEXT("Cannot call GetSpawnables on a null sequence"), ELogVerbosity::Error);
+		return TArray<FSequencerBindingProxy>();
+	}
+
 	TArray<FSequencerBindingProxy> AllSpawnables;
 
 	UMovieScene* MovieScene = GetMovieScene(Sequence);
@@ -514,6 +748,12 @@ TArray<FSequencerBindingProxy> UMovieSceneSequenceExtensions::GetSpawnables(UMov
 
 TArray<FSequencerBindingProxy> UMovieSceneSequenceExtensions::GetPossessables(UMovieSceneSequence* Sequence)
 {
+	if (!Sequence)
+	{
+		FFrame::KismetExecutionMessage(TEXT("Cannot call GetPossessables on a null sequence"), ELogVerbosity::Error);
+		return TArray<FSequencerBindingProxy>();
+	}
+
 	TArray<FSequencerBindingProxy> AllPossessables;
 
 	UMovieScene* MovieScene = GetMovieScene(Sequence);
@@ -532,37 +772,82 @@ TArray<FSequencerBindingProxy> UMovieSceneSequenceExtensions::GetPossessables(UM
 
 FSequencerBindingProxy UMovieSceneSequenceExtensions::AddPossessable(UMovieSceneSequence* Sequence, UObject* ObjectToPossess)
 {
+	if (!Sequence)
+	{
+		FFrame::KismetExecutionMessage(TEXT("Cannot call AddPossessable on a null sequence"), ELogVerbosity::Error);
+		return FSequencerBindingProxy();
+	}
+
+	if (!ObjectToPossess)
+	{
+		FFrame::KismetExecutionMessage(TEXT("Cannot call AddPossessable on a null object"), ELogVerbosity::Error);
+		return FSequencerBindingProxy();
+	}
+
 	FGuid NewGuid = Sequence->CreatePossessable(ObjectToPossess);
 	return FSequencerBindingProxy(NewGuid, Sequence);
 }
 
 FSequencerBindingProxy UMovieSceneSequenceExtensions::AddSpawnableFromInstance(UMovieSceneSequence* Sequence, UObject* ObjectToSpawn)
 {
+	if (!Sequence)
+	{
+		FFrame::KismetExecutionMessage(TEXT("Cannot call AddSpawnableFromInstance on a null sequence"), ELogVerbosity::Error);
+		return FSequencerBindingProxy();
+	}
+
+	if (!ObjectToSpawn)
+	{
+		FFrame::KismetExecutionMessage(TEXT("Cannot call AddSpawnableFromInstance on a null object"), ELogVerbosity::Error);
+		return FSequencerBindingProxy();
+	}
+
 	FGuid NewGuid = Sequence->AllowsSpawnableObjects() ? Sequence->CreateSpawnable(ObjectToSpawn) : FGuid();
 	return FSequencerBindingProxy(NewGuid, Sequence);
 }
 
 FSequencerBindingProxy UMovieSceneSequenceExtensions::AddSpawnableFromClass(UMovieSceneSequence* Sequence, UClass* ClassToSpawn)
 {
+	if (!Sequence)
+	{
+		FFrame::KismetExecutionMessage(TEXT("Cannot call AddSpawnableFromClass on a null sequence"), ELogVerbosity::Error);
+		return FSequencerBindingProxy();
+	}
+
+	if (!ClassToSpawn)
+	{
+		FFrame::KismetExecutionMessage(TEXT("Cannot call AddSpawnableFromClass on a null class"), ELogVerbosity::Error);
+		return FSequencerBindingProxy();
+	}
+
 	FGuid NewGuid = Sequence->AllowsSpawnableObjects() ? Sequence->CreateSpawnable(ClassToSpawn) : FGuid();
 	return FSequencerBindingProxy(NewGuid, Sequence);
 }
 
 TArray<UObject*> UMovieSceneSequenceExtensions::LocateBoundObjects(UMovieSceneSequence* Sequence, const FSequencerBindingProxy& InBinding, UObject* Context)
 {
-	TArray<UObject*> Result;
-	if (Sequence)
+	if (!Sequence)
 	{
-		TArray<UObject*, TInlineAllocator<1>> OutObjects;
-		Sequence->LocateBoundObjects(InBinding.BindingID, Context, OutObjects);
-		Result.Append(OutObjects);
+		FFrame::KismetExecutionMessage(TEXT("Cannot call LocateBoundObjects on a null sequence"), ELogVerbosity::Error);
+		return TArray<UObject*>();
 	}
 
+	TArray<UObject*> Result;
+	TArray<UObject*, TInlineAllocator<1>> OutObjects;
+	Sequence->LocateBoundObjects(InBinding.BindingID, Context, OutObjects);
+	Result.Append(OutObjects);
+
 	return Result;
 }
 
 FMovieSceneObjectBindingID UMovieSceneSequenceExtensions::MakeBindingID(UMovieSceneSequence* MasterSequence, const FSequencerBindingProxy& InBinding, EMovieSceneObjectBindingSpace Space)
 {
+	if (!MasterSequence)
+	{
+		FFrame::KismetExecutionMessage(TEXT("Cannot call MakeBindingID on a null sequence"), ELogVerbosity::Error);
+		return FMovieSceneObjectBindingID();
+	}
+
 	// This function was kinda flawed before - when ::Local was passed for the Space parameter,
 	// and the sub sequence ID could not be found it would always fall back to a binding for ::Root without any Sequence ID
 	FMovieSceneObjectBindingID BindingID = GetPortableBindingID(MasterSequence, MasterSequence, InBinding);
@@ -582,7 +867,7 @@ FMovieSceneObjectBindingID UMovieSceneSequenceExtensions::GetPortableBindingID(U
 {
 	if (!MasterSequence || !DestinationSequence || !InBinding.Sequence)
 	{
-		FFrame::KismetExecutionMessage(TEXT("Invalid sequence sepcified."), ELogVerbosity::Error);
+		FFrame::KismetExecutionMessage(TEXT("Invalid sequence specified."), ELogVerbosity::Error);
 		return FMovieSceneObjectBindingID();
 	}
 
@@ -659,6 +944,12 @@ FMovieSceneObjectBindingID UMovieSceneSequenceExtensions::GetPortableBindingID(U
 
 FSequencerBindingProxy UMovieSceneSequenceExtensions::ResolveBindingID(UMovieSceneSequence* MasterSequence, FMovieSceneObjectBindingID InObjectBindingID)
 {
+	if (!MasterSequence)
+	{
+		FFrame::KismetExecutionMessage(TEXT("Cannot call ResolveBindingID on a null sequence"), ELogVerbosity::Error);
+		return FSequencerBindingProxy();
+	}
+
 	UMovieSceneSequence* Sequence = MasterSequence;
 
 	FMovieSceneCompiledDataID DataID = UMovieSceneCompiledDataManager::GetPrecompiledData()->Compile(MasterSequence);
@@ -677,16 +968,19 @@ FSequencerBindingProxy UMovieSceneSequenceExtensions::ResolveBindingID(UMovieSce
 
 TArray<UMovieSceneFolder*> UMovieSceneSequenceExtensions::GetRootFoldersInSequence(UMovieSceneSequence* Sequence)
 {
+	if (!Sequence)
+	{
+		FFrame::KismetExecutionMessage(TEXT("Cannot call GetRootFoldersInSequence on a null sequence"), ELogVerbosity::Error);
+		return TArray<UMovieSceneFolder*>();
+	}
+
 	TArray<UMovieSceneFolder*> Result;
 
 #if WITH_EDITORONLY_DATA
-	if (Sequence)
+	UMovieScene* Scene = Sequence->GetMovieScene();
+	if (Scene)
 	{
-		UMovieScene* Scene = Sequence->GetMovieScene();
-		if (Scene)
-		{
-			Result = Scene->GetRootFolders();
-		}
+		Result = Scene->GetRootFolders();
 	}
 #endif
 
@@ -695,19 +989,23 @@ TArray<UMovieSceneFolder*> UMovieSceneSequenceExtensions::GetRootFoldersInSequen
 
 UMovieSceneFolder* UMovieSceneSequenceExtensions::AddRootFolderToSequence(UMovieSceneSequence* Sequence, FString NewFolderName)
 {
+	if (!Sequence)
+	{
+		FFrame::KismetExecutionMessage(TEXT("Cannot call AddRootFolderToSequence on a null sequence"), ELogVerbosity::Error);
+		return nullptr;
+	}
+
 	UMovieSceneFolder* NewFolder = nullptr;
 	
 #if WITH_EDITORONLY_DATA
-	if (Sequence)
+	
+	UMovieScene* MovieScene = Sequence->GetMovieScene();
+	if (MovieScene)
 	{
-		UMovieScene* MovieScene = Sequence->GetMovieScene();
-		if (MovieScene)
-		{
-			MovieScene->Modify();
-			NewFolder = NewObject<UMovieSceneFolder>(MovieScene);
-			NewFolder->SetFolderName(FName(*NewFolderName));
-			MovieScene->GetRootFolders().Add(NewFolder);
-		}
+		MovieScene->Modify();
+		NewFolder = NewObject<UMovieSceneFolder>(MovieScene);
+		NewFolder->SetFolderName(FName(*NewFolderName));
+		MovieScene->GetRootFolders().Add(NewFolder);
 	}
 #endif
 
@@ -716,6 +1014,12 @@ UMovieSceneFolder* UMovieSceneSequenceExtensions::AddRootFolderToSequence(UMovie
 
 TArray<FMovieSceneMarkedFrame> UMovieSceneSequenceExtensions::GetMarkedFrames(UMovieSceneSequence* Sequence)
 {
+	if (!Sequence)
+	{
+		FFrame::KismetExecutionMessage(TEXT("Cannot call GetMarkedFrames on a null sequence"), ELogVerbosity::Error);
+		return TArray<FMovieSceneMarkedFrame>();
+	}
+
 	UMovieScene* MovieScene = Sequence->GetMovieScene();
 	if (MovieScene)
 	{
@@ -727,6 +1031,12 @@ TArray<FMovieSceneMarkedFrame> UMovieSceneSequenceExtensions::GetMarkedFrames(UM
 
 int32 UMovieSceneSequenceExtensions::AddMarkedFrame(UMovieSceneSequence* Sequence, const FMovieSceneMarkedFrame& InMarkedFrame)
 {
+	if (!Sequence)
+	{
+		FFrame::KismetExecutionMessage(TEXT("Cannot call AddMarkedFrame on a null sequence"), ELogVerbosity::Error);
+		return INDEX_NONE;
+	}
+
 	UMovieScene* MovieScene = Sequence->GetMovieScene();
 	if (MovieScene)
 	{
@@ -740,6 +1050,12 @@ int32 UMovieSceneSequenceExtensions::AddMarkedFrame(UMovieSceneSequence* Sequenc
 
 void UMovieSceneSequenceExtensions::SetMarkedFrame(UMovieSceneSequence* Sequence, int32 InMarkIndex, FFrameNumber InFrameNumber)
 {
+	if (!Sequence)
+	{
+		FFrame::KismetExecutionMessage(TEXT("Cannot call SetMarkedFrame on a null sequence"), ELogVerbosity::Error);
+		return;
+	}
+
 	UMovieScene* MovieScene = Sequence->GetMovieScene();
 	if (MovieScene)
 	{
@@ -750,6 +1066,12 @@ void UMovieSceneSequenceExtensions::SetMarkedFrame(UMovieSceneSequence* Sequence
 
 void UMovieSceneSequenceExtensions::DeleteMarkedFrame(UMovieSceneSequence* Sequence, int32 DeleteIndex)
 {
+	if (!Sequence)
+	{
+		FFrame::KismetExecutionMessage(TEXT("Cannot call DeleteMarkedFrame on a null sequence"), ELogVerbosity::Error);
+		return;
+	}
+
 	UMovieScene* MovieScene = Sequence->GetMovieScene();
 	if (MovieScene)
 	{
@@ -761,6 +1083,12 @@ void UMovieSceneSequenceExtensions::DeleteMarkedFrame(UMovieSceneSequence* Seque
 
 void UMovieSceneSequenceExtensions::DeleteMarkedFrames(UMovieSceneSequence* Sequence)
 {
+	if (!Sequence)
+	{
+		FFrame::KismetExecutionMessage(TEXT("Cannot call DeleteMarkedFrames on a null sequence"), ELogVerbosity::Error);
+		return;
+	}
+
 	UMovieScene* MovieScene = Sequence->GetMovieScene();
 	if (MovieScene)
 	{
@@ -772,6 +1100,12 @@ void UMovieSceneSequenceExtensions::DeleteMarkedFrames(UMovieSceneSequence* Sequ
 
 void UMovieSceneSequenceExtensions::SortMarkedFrames(UMovieSceneSequence* Sequence)
 {
+	if (!Sequence)
+	{
+		FFrame::KismetExecutionMessage(TEXT("Cannot call SortMarkedFrames on a null sequence"), ELogVerbosity::Error);
+		return;
+	}
+
 	UMovieScene* MovieScene = Sequence->GetMovieScene();
 	if (MovieScene)
 	{
@@ -781,6 +1115,12 @@ void UMovieSceneSequenceExtensions::SortMarkedFrames(UMovieSceneSequence* Sequen
 
 int32 UMovieSceneSequenceExtensions::FindMarkedFrameByLabel(UMovieSceneSequence* Sequence, const FString& InLabel)
 {
+	if (!Sequence)
+	{
+		FFrame::KismetExecutionMessage(TEXT("Cannot call FindMarkedFrameByLabel on a null sequence"), ELogVerbosity::Error);
+		return INDEX_NONE;
+	}
+
 	UMovieScene* MovieScene = Sequence->GetMovieScene();
 	if (MovieScene)
 	{
@@ -791,6 +1131,12 @@ int32 UMovieSceneSequenceExtensions::FindMarkedFrameByLabel(UMovieSceneSequence*
 
 int32 UMovieSceneSequenceExtensions::FindMarkedFrameByFrameNumber(UMovieSceneSequence* Sequence, FFrameNumber InFrameNumber)
 {
+	if (!Sequence)
+	{
+		FFrame::KismetExecutionMessage(TEXT("Cannot call FindMarkedFrameByFrameNumber on a null sequence"), ELogVerbosity::Error);
+		return INDEX_NONE;
+	}
+
 	UMovieScene* MovieScene = Sequence->GetMovieScene();
 	if (MovieScene)
 	{
@@ -801,6 +1147,12 @@ int32 UMovieSceneSequenceExtensions::FindMarkedFrameByFrameNumber(UMovieSceneSeq
 
 int32 UMovieSceneSequenceExtensions::FindNextMarkedFrame(UMovieSceneSequence* Sequence, FFrameNumber InFrameNumber, bool bForward)
 {
+	if (!Sequence)
+	{
+		FFrame::KismetExecutionMessage(TEXT("Cannot call FindNextMarkedFrame on a null sequence"), ELogVerbosity::Error);
+		return INDEX_NONE;
+	}
+
 	UMovieScene* MovieScene = Sequence->GetMovieScene();
 	if (MovieScene)
 	{
@@ -811,6 +1163,12 @@ int32 UMovieSceneSequenceExtensions::FindNextMarkedFrame(UMovieSceneSequence* Se
 
 void UMovieSceneSequenceExtensions::SetReadOnly(UMovieSceneSequence* Sequence, bool bInReadOnly)
 {
+	if (!Sequence)
+	{
+		FFrame::KismetExecutionMessage(TEXT("Cannot call SetReadOnly on a null sequence"), ELogVerbosity::Error);
+		return;
+	}
+
 #if WITH_EDITORONLY_DATA
 	UMovieScene* MovieScene = Sequence->GetMovieScene();
 	if (MovieScene)
@@ -822,6 +1180,12 @@ void UMovieSceneSequenceExtensions::SetReadOnly(UMovieSceneSequence* Sequence, b
 
 bool UMovieSceneSequenceExtensions::IsReadOnly(UMovieSceneSequence* Sequence)
 {
+	if (!Sequence)
+	{
+		FFrame::KismetExecutionMessage(TEXT("Cannot call IsReadOnly on a null sequence"), ELogVerbosity::Error);
+		return false;
+	}
+
 #if WITH_EDITORONLY_DATA
 	UMovieScene* MovieScene = Sequence->GetMovieScene();
 	if (MovieScene)
diff --git a/Engine/Plugins/MovieScene/SequencerScripting/Source/SequencerScripting/Private/ExtensionLibraries/MovieSceneTrackExtensions.cpp b/Engine/Plugins/MovieScene/SequencerScripting/Source/SequencerScripting/Private/ExtensionLibraries/MovieSceneTrackExtensions.cpp
--- a/Engine/Plugins/MovieScene/SequencerScripting/Source/SequencerScripting/Private/ExtensionLibraries/MovieSceneTrackExtensions.cpp
+++ b/Engine/Plugins/MovieScene/SequencerScripting/Source/SequencerScripting/Private/ExtensionLibraries/MovieSceneTrackExtensions.cpp
@@ -6,6 +6,12 @@
 
 void UMovieSceneTrackExtensions::SetDisplayName(UMovieSceneTrack* Track, const FText& InName)
 {
+	if (!Track)
+	{
+		FFrame::KismetExecutionMessage(TEXT("Cannot call SetDisplayName on a null track"), ELogVerbosity::Error);
+		return;
+	}
+
 	if (UMovieSceneNameableTrack* NameableTrack = Cast<UMovieSceneNameableTrack>(Track))
 	{
 #if WITH_EDITORONLY_DATA
@@ -16,6 +22,12 @@ void UMovieSceneTrackExtensions::SetDisplayName(UMovieSceneTrack* Track, const F
 
 FText UMovieSceneTrackExtensions::GetDisplayName(UMovieSceneTrack* Track)
 {
+	if (!Track)
+	{
+		FFrame::KismetExecutionMessage(TEXT("Cannot call GetDisplayName on a null track"), ELogVerbosity::Error);
+		return FText::GetEmpty();
+	}
+
 #if WITH_EDITORONLY_DATA
 	return Track->GetDisplayName();
 #endif
@@ -24,6 +36,12 @@ FText UMovieSceneTrackExtensions::GetDisplayName(UMovieSceneTrack* Track)
 
 void UMovieSceneTrackExtensions::SetTrackRowDisplayName(UMovieSceneTrack* Track, const FText& InName, int32 RowIndex)
 {
+	if (!Track)
+	{
+		FFrame::KismetExecutionMessage(TEXT("Cannot call SetTrackRowDisplayName on a null track"), ELogVerbosity::Error);
+		return;
+	}
+
 	if (UMovieSceneNameableTrack* NameableTrack = Cast<UMovieSceneNameableTrack>(Track))
 	{
 #if WITH_EDITORONLY_DATA
@@ -34,6 +52,12 @@ void UMovieSceneTrackExtensions::SetTrackRowDisplayName(UMovieSceneTrack* Track,
 
 FText UMovieSceneTrackExtensions::GetTrackRowDisplayName(UMovieSceneTrack* Track, int32 RowIndex)
 {
+	if (!Track)
+	{
+		FFrame::KismetExecutionMessage(TEXT("Cannot call GetTrackRowDisplayName on a null track"), ELogVerbosity::Error);
+		return FText::GetEmpty();
+	}
+
 #if WITH_EDITORONLY_DATA
 	return Track->GetTrackRowDisplayName(RowIndex);
 #endif
@@ -42,6 +66,12 @@ FText UMovieSceneTrackExtensions::GetTrackRowDisplayName(UMovieSceneTrack* Track
 
 UMovieSceneSection* UMovieSceneTrackExtensions::AddSection(UMovieSceneTrack* Track)
 {
+	if (!Track)
+	{
+		FFrame::KismetExecutionMessage(TEXT("Cannot call AddSection on a null track"), ELogVerbosity::Error);
+		return nullptr;
+	}
+
 	UMovieSceneSection* NewSection = Track->CreateNewSection();
 
 	if (NewSection)
@@ -56,11 +86,23 @@ UMovieSceneSection* UMovieSceneTrackExtensions::AddSection(UMovieSceneTrack* Tra
 
 TArray<UMovieSceneSection*> UMovieSceneTrackExtensions::GetSections(UMovieSceneTrack* Track)
 {
+	if (!Track)
+	{
+		FFrame::KismetExecutionMessage(TEXT("Cannot call GetSections on a null track"), ELogVerbosity::Error);
+		return TArray<UMovieSceneSection*>();
+	}
+
 	return Track->GetAllSections();
 }
 
 void UMovieSceneTrackExtensions::RemoveSection(UMovieSceneTrack* Track, UMovieSceneSection* Section)
 {
+	if (!Track)
+	{
+		FFrame::KismetExecutionMessage(TEXT("Cannot call RemoveSection on a null track"), ELogVerbosity::Error);
+		return;
+	}
+
 	if (Section)
 	{
 		Track->Modify();
@@ -71,6 +113,12 @@ void UMovieSceneTrackExtensions::RemoveSection(UMovieSceneTrack* Track, UMovieSc
 
 int32 UMovieSceneTrackExtensions::GetSortingOrder(UMovieSceneTrack* Track) 
 { 
+	if (!Track)
+	{
+		FFrame::KismetExecutionMessage(TEXT("Cannot call GetSortingOrder on a null track"), ELogVerbosity::Error);
+		return -1;
+	}
+
 #if WITH_EDITORONLY_DATA
 	return Track->GetSortingOrder(); 
 #endif
@@ -79,6 +127,12 @@ int32 UMovieSceneTrackExtensions::GetSortingOrder(UMovieSceneTrack* Track)
  
 void UMovieSceneTrackExtensions::SetSortingOrder(UMovieSceneTrack* Track, int32 SortingOrder) 
 {
+	if (!Track)
+	{
+		FFrame::KismetExecutionMessage(TEXT("Cannot call SetSortingOrder on a null track"), ELogVerbosity::Error);
+		return;
+	}
+
 #if WITH_EDITORONLY_DATA
 	Track->SetSortingOrder(SortingOrder); 
 #endif
@@ -86,6 +140,12 @@ void UMovieSceneTrackExtensions::SetSortingOrder(UMovieSceneTrack* Track, int32
 
 FColor UMovieSceneTrackExtensions::GetColorTint(UMovieSceneTrack* Track) 
 { 
+	if (!Track)
+	{
+		FFrame::KismetExecutionMessage(TEXT("Cannot call GetColorTint on a null track"), ELogVerbosity::Error);
+		return FColor();
+	}
+
 #if WITH_EDITORONLY_DATA
 	return Track->GetColorTint(); 
 #endif
@@ -94,6 +154,12 @@ FColor UMovieSceneTrackExtensions::GetColorTint(UMovieSceneTrack* Track)
 
 void UMovieSceneTrackExtensions::SetColorTint(UMovieSceneTrack* Track, const FColor& ColorTint) 
 { 
+	if (!Track)
+	{
+		FFrame::KismetExecutionMessage(TEXT("Cannot call SetColorTint on a null track"), ELogVerbosity::Error);
+		return;
+	}
+
 #if WITH_EDITORONLY_DATA
 	Track->SetColorTint(ColorTint); 
 #endif
@@ -101,11 +167,23 @@ void UMovieSceneTrackExtensions::SetColorTint(UMovieSceneTrack* Track, const FCo
 
 UMovieSceneSection* UMovieSceneTrackExtensions::GetSectionToKey(UMovieSceneTrack* Track) 
 { 
+	if (!Track)
+	{
+		FFrame::KismetExecutionMessage(TEXT("Cannot call GetSectionToKey on a null track"), ELogVerbosity::Error);
+		return nullptr;
+	}
+
 	return Track->GetSectionToKey(); 
 }
 
 void UMovieSceneTrackExtensions::SetSectionToKey(UMovieSceneTrack* Track, UMovieSceneSection* Section) 
 { 
+	if (!Track)
+	{
+		FFrame::KismetExecutionMessage(TEXT("Cannot call SetSectionToKey on a null track"), ELogVerbosity::Error);
+		return;
+	}
+
 	Track->SetSectionToKey(Section); 
 }
 
diff --git a/Engine/Plugins/MovieScene/SequencerScripting/Source/SequencerScripting/Private/ExtensionLibraries/MovieSceneVectorTrackExtensions.cpp b/Engine/Plugins/MovieScene/SequencerScripting/Source/SequencerScripting/Private/ExtensionLibraries/MovieSceneVectorTrackExtensions.cpp
--- a/Engine/Plugins/MovieScene/SequencerScripting/Source/SequencerScripting/Private/ExtensionLibraries/MovieSceneVectorTrackExtensions.cpp
+++ b/Engine/Plugins/MovieScene/SequencerScripting/Source/SequencerScripting/Private/ExtensionLibraries/MovieSceneVectorTrackExtensions.cpp
@@ -5,6 +5,12 @@
 
 void UMovieSceneFloatVectorTrackExtensions::SetNumChannelsUsed(UMovieSceneFloatVectorTrack* Track, int32 InNumChannelsUsed)
 {
+	if (!Track)
+	{
+		FFrame::KismetExecutionMessage(TEXT("Cannot call SetNumChannelsUsed on a null track"), ELogVerbosity::Error);
+		return;
+	}
+
 	Track->Modify();
 
 	Track->SetNumChannelsUsed(InNumChannelsUsed);
@@ -13,11 +19,23 @@ void UMovieSceneFloatVectorTrackExtensions::SetNumChannelsUsed(UMovieSceneFloatV
 
 int32 UMovieSceneFloatVectorTrackExtensions::GetNumChannelsUsed(UMovieSceneFloatVectorTrack* Track)
 {
+	if (!Track)
+	{
+		FFrame::KismetExecutionMessage(TEXT("Cannot call GetNumChannelsUsed on a null track"), ELogVerbosity::Error);
+		return -1;
+	}
+
 	return Track->GetNumChannelsUsed();
 }
 
 void UMovieSceneDoubleVectorTrackExtensions::SetNumChannelsUsed(UMovieSceneDoubleVectorTrack* Track, int32 InNumChannelsUsed)
 {
+	if (!Track)
+	{
+		FFrame::KismetExecutionMessage(TEXT("Cannot call SetNumChannelsUsed on a null track"), ELogVerbosity::Error);
+		return;
+	}
+
 	Track->Modify();
 
 	Track->SetNumChannelsUsed(InNumChannelsUsed);
@@ -26,6 +44,12 @@ void UMovieSceneDoubleVectorTrackExtensions::SetNumChannelsUsed(UMovieSceneDoubl
 
 int32 UMovieSceneDoubleVectorTrackExtensions::GetNumChannelsUsed(UMovieSceneDoubleVectorTrack* Track)
 {
+	if (!Track)
+	{
+		FFrame::KismetExecutionMessage(TEXT("Cannot call GetNumChannelsUsed on a null track"), ELogVerbosity::Error);
+		return -1;
+	}
+
 	return Track->GetNumChannelsUsed();
 }
 
