Sequencer: Added a method for deleting root folders from a sequence

diff --git a/Engine/Plugins/MovieScene/SequencerScripting/Source/SequencerScripting/Private/ExtensionLibraries/MovieSceneSequenceExtensions.cpp b/Engine/Plugins/MovieScene/SequencerScripting/Source/SequencerScripting/Private/ExtensionLibraries/MovieSceneSequenceExtensions.cpp
--- a/Engine/Plugins/MovieScene/SequencerScripting/Source/SequencerScripting/Private/ExtensionLibraries/MovieSceneSequenceExtensions.cpp
+++ b/Engine/Plugins/MovieScene/SequencerScripting/Source/SequencerScripting/Private/ExtensionLibraries/MovieSceneSequenceExtensions.cpp
@@ -1012,6 +1012,36 @@ UMovieSceneFolder* UMovieSceneSequenceExtensions::AddRootFolderToSequence(UMovie
 	return NewFolder;
 }
 
+void UMovieSceneSequenceExtensions::RemoveRootFolderFromSequence(UMovieSceneSequence* Sequence, UMovieSceneFolder* Folder)
+{
+	if (!Sequence || !Folder)
+	{
+		FFrame::KismetExecutionMessage(TEXT("Cannot call RemoveRootFolderFromSequence on a null sequence or folder"), ELogVerbosity::Error);
+		return;
+	}
+
+	UMovieScene* FolderMovieScene = Folder->GetTypedOuter<UMovieScene>();
+	if (FolderMovieScene != Sequence->GetMovieScene())
+	{
+		FFrame::KismetExecutionMessage(*FString::Printf(TEXT("The folder '%s' does not belong to sequence '%s'"),
+			*Folder->GetFolderName().ToString(), *Sequence->GetName()), ELogVerbosity::Error);
+		return;
+	}
+
+#if WITH_EDITORONLY_DATA
+	if (FolderMovieScene)
+	{
+		FolderMovieScene->Modify();
+		const int32 NumFoldersRemoved = FolderMovieScene->GetRootFolders().Remove(Folder);
+		if (NumFoldersRemoved == 0)
+		{
+			FFrame::KismetExecutionMessage(*FString::Printf(TEXT("The specified folder '%s' is not a root folder"),
+				*Folder->GetFolderName().ToString()), ELogVerbosity::Error);
+		}
+	}
+#endif
+}
+
 TArray<FMovieSceneMarkedFrame> UMovieSceneSequenceExtensions::GetMarkedFrames(UMovieSceneSequence* Sequence)
 {
 	if (!Sequence)
diff --git a/Engine/Plugins/MovieScene/SequencerScripting/Source/SequencerScripting/Public/ExtensionLibraries/MovieSceneSequenceExtensions.h b/Engine/Plugins/MovieScene/SequencerScripting/Source/SequencerScripting/Public/ExtensionLibraries/MovieSceneSequenceExtensions.h
--- a/Engine/Plugins/MovieScene/SequencerScripting/Source/SequencerScripting/Public/ExtensionLibraries/MovieSceneSequenceExtensions.h
+++ b/Engine/Plugins/MovieScene/SequencerScripting/Source/SequencerScripting/Public/ExtensionLibraries/MovieSceneSequenceExtensions.h
@@ -503,6 +503,15 @@ public:
 	UFUNCTION(BlueprintCallable, Category="Sequencer|Sequence", meta=(ScriptMethod))
 	static UMovieSceneFolder* AddRootFolderToSequence(UMovieSceneSequence* Sequence, FString NewFolderName);
 
+	/**
+	 * Remove a root folder from the given sequence. Will throw an exception if the specified folder is not valid or not a root folder.
+	 *
+	 * @param Sequence			The sequence That the folder belongs to
+	 * @param Folder			The folder to remove
+	 */
+	UFUNCTION(BlueprintCallable, Category="Sequencer|Sequence", meta=(ScriptMethod))
+	static void RemoveRootFolderFromSequence(UMovieSceneSequence* Sequence, UMovieSceneFolder* Folder);
+
 public:
 
 	/**
