Sequencer: Linked Anim Sequences: When duplicating level sequences we also duplicate the linked anim sequence.  Also added some BP/Python functions for creating linked anim sequences

diff --git a/Engine/Plugins/MovieScene/SequencerScripting/Source/SequencerScriptingEditor/Private/SequencerTools.cpp b/Engine/Plugins/MovieScene/SequencerScripting/Source/SequencerScriptingEditor/Private/SequencerTools.cpp
--- a/Engine/Plugins/MovieScene/SequencerScripting/Source/SequencerScriptingEditor/Private/SequencerTools.cpp
+++ b/Engine/Plugins/MovieScene/SequencerScripting/Source/SequencerScriptingEditor/Private/SequencerTools.cpp
@@ -30,6 +30,11 @@
 #include "BlueprintActionMenuItem.h"
 #include "EdGraphSchema_K2.h"
 
+#include "AssetData.h"
+#include "LevelSequenceAnimSequenceLink.h"
+#include "AnimSequenceLevelSequenceLink.h"
+
+
 #define LOCTEXT_NAMESPACE "SequencerTools"
 
 bool USequencerToolsFunctionLibrary::RenderMovie(UMovieSceneCapture* InCaptureSettings, FOnRenderMovieStopped OnFinishedCallback)
@@ -283,7 +288,7 @@ static USkeletalMeshComponent* GetSkelMeshComponent(IMovieScenePlayer* Player, c
 	return nullptr;
 }
 
-bool USequencerToolsFunctionLibrary::ExportAnimSequence(UWorld* World, ULevelSequence*  Sequence,  UAnimSequence* AnimSequence, UAnimSeqExportOption* ExportOptions,const FSequencerBindingProxy& Binding)
+bool USequencerToolsFunctionLibrary::ExportAnimSequence(UWorld* World, ULevelSequence*  Sequence,  UAnimSequence* AnimSequence, UAnimSeqExportOption* ExportOptions,const FSequencerBindingProxy& Binding, bool bCreateLink)
 {
 	UMovieScene* MovieScene = Sequence->GetMovieScene();
 	if (Binding.Sequence != Sequence || !AnimSequence)
@@ -322,10 +327,101 @@ bool USequencerToolsFunctionLibrary::ExportAnimSequence(UWorld* World, ULevelSeq
 	Player->Stop();
 	World->DestroyActor(OutActor);
 
+	//create the link to the anim sequence
+	if (bResult && bCreateLink)
+	{
+		//create from anim sequence back to level sequence
+		if (IInterface_AssetUserData* AnimAssetUserData = Cast< IInterface_AssetUserData >(AnimSequence))
+		{
+			UAnimSequenceLevelSequenceLink* AnimLevelLink = AnimAssetUserData->GetAssetUserData< UAnimSequenceLevelSequenceLink >();
+			if (!AnimLevelLink)
+			{
+				AnimLevelLink = NewObject<UAnimSequenceLevelSequenceLink>(AnimSequence, NAME_None, RF_Public | RF_Transactional);
+				AnimAssetUserData->AddAssetUserData(AnimLevelLink);
+			}
+
+			AnimLevelLink->SetLevelSequence(Sequence);
+			AnimLevelLink->SkelTrackGuid = Binding.BindingID;
+		}
+		//create from level sequence to anim sequence, trickier since we can have multiples here.
+		if (IInterface_AssetUserData* AssetUserDataInterface = Cast< IInterface_AssetUserData >(Sequence))
+		{
+			bool bAddItem = true;
+			ULevelSequenceAnimSequenceLink* LevelAnimLink = AssetUserDataInterface->GetAssetUserData< ULevelSequenceAnimSequenceLink >();
+			if (LevelAnimLink)
+			{
+				for (FLevelSequenceAnimSequenceLinkItem& LevelAnimLinkItem : LevelAnimLink->AnimSequenceLinks)
+				{
+					if (LevelAnimLinkItem.SkelTrackGuid == Binding.BindingID)
+					{
+						bAddItem = false;
+						UAnimSequence* OtherAnimSequence = LevelAnimLinkItem.ResolveAnimSequence();
+
+						if (OtherAnimSequence != AnimSequence)
+						{
+							if (IInterface_AssetUserData* OtherAnimAssetUserData = Cast< IInterface_AssetUserData >(OtherAnimSequence))
+							{
+								UAnimSequenceLevelSequenceLink* OtherAnimLevelLink = OtherAnimAssetUserData->GetAssetUserData< UAnimSequenceLevelSequenceLink >();
+								if (OtherAnimLevelLink)
+								{
+									OtherAnimAssetUserData->RemoveUserDataOfClass(UAnimSequenceLevelSequenceLink::StaticClass());
+								}
+							}
+						}
+						LevelAnimLinkItem.PathToAnimSequence = FSoftObjectPath(AnimSequence);
+						LevelAnimLinkItem.bExportMorphTargets = ExportOptions->bExportMorphTargets;
+						LevelAnimLinkItem.bExportAttributeCurves = ExportOptions->bExportAttributeCurves;
+						LevelAnimLinkItem.bExportMaterialCurves = ExportOptions->bExportMaterialCurves;
+						LevelAnimLinkItem.bExportTransforms = ExportOptions->bExportTransforms;
+						LevelAnimLinkItem.bRecordInWorldSpace = ExportOptions->bRecordInWorldSpace;
+
+						break;
+					}
+				}
+			}
+			else
+			{
+				LevelAnimLink = NewObject<ULevelSequenceAnimSequenceLink>(Sequence, NAME_None, RF_Public | RF_Transactional);
+
+			}
+			if (bAddItem == true)
+			{
+				FLevelSequenceAnimSequenceLinkItem LevelAnimLinkItem;
+				LevelAnimLinkItem.SkelTrackGuid = Binding.BindingID;
+				LevelAnimLinkItem.PathToAnimSequence = FSoftObjectPath(AnimSequence);
+				LevelAnimLinkItem.bExportMorphTargets = ExportOptions->bExportMorphTargets;
+				LevelAnimLinkItem.bExportAttributeCurves = ExportOptions->bExportAttributeCurves;
+				LevelAnimLinkItem.bExportMaterialCurves = ExportOptions->bExportMaterialCurves;
+				LevelAnimLinkItem.bExportTransforms = ExportOptions->bExportTransforms;
+				LevelAnimLinkItem.bRecordInWorldSpace = ExportOptions->bRecordInWorldSpace;
+
+				LevelAnimLink->AnimSequenceLinks.Add(LevelAnimLinkItem);
+				AssetUserDataInterface->AddAssetUserData(LevelAnimLink);
+			}
+		}
+	}
 	return bResult;
 }
 
+UAnimSequenceLevelSequenceLink* USequencerToolsFunctionLibrary::GetLevelSequenceLinkFromAnimSequence(UAnimSequence* InAnimSequence)
+{
+	if (IInterface_AssetUserData* AnimAssetUserData = Cast< IInterface_AssetUserData >(InAnimSequence))
+	{
+		UAnimSequenceLevelSequenceLink* AnimLevelLink = AnimAssetUserData->GetAssetUserData< UAnimSequenceLevelSequenceLink >();
+		return AnimLevelLink;
+	}
+	return nullptr;
+}
 
+ULevelSequenceAnimSequenceLink* USequencerToolsFunctionLibrary::GetAnimSequenceLinkFromLevelSequence(ULevelSequence* InLevelSequence)
+{
+	if (IInterface_AssetUserData* AssetUserDataInterface = Cast< IInterface_AssetUserData >(InLevelSequence))
+	{
+		ULevelSequenceAnimSequenceLink* LevelAnimLink = AssetUserDataInterface->GetAssetUserData< ULevelSequenceAnimSequenceLink >();
+		return LevelAnimLink;
+	}
+	return nullptr;
+}
 
 TArray<FGuid> AddActors(UWorld* World, UMovieSceneSequence* InSequence, UMovieScene* InMovieScene, IMovieScenePlayer* Player, FMovieSceneSequenceIDRef TemplateID,const TArray<TWeakObjectPtr<AActor> >& InActors)
 {
diff --git a/Engine/Plugins/MovieScene/SequencerScripting/Source/SequencerScriptingEditor/Public/SequencerTools.h b/Engine/Plugins/MovieScene/SequencerScripting/Source/SequencerScriptingEditor/Public/SequencerTools.h
--- a/Engine/Plugins/MovieScene/SequencerScripting/Source/SequencerScriptingEditor/Public/SequencerTools.h
+++ b/Engine/Plugins/MovieScene/SequencerScripting/Source/SequencerScriptingEditor/Public/SequencerTools.h
@@ -23,8 +23,13 @@ class UK2Node_CustomEvent;
 
 class UAnimSeqExportOption;
 class UMovieSceneUserImportFBXControlRigSettings;
+
+class ULevelSequenceAnimSequenceLink;
+class UAnimSequenceLevelSequenceLink;
+
 DECLARE_DYNAMIC_DELEGATE_OneParam(FOnRenderMovieStopped, bool, bSuccess);
 
+
 USTRUCT(BlueprintType)
 struct SEQUENCERSCRIPTINGEDITOR_API FSequencerBoundObjects
 {
@@ -136,9 +141,28 @@ public:
 	 * @ExportOption The export options for the sequence.
 	 * @InBinding Binding to export that has a skelmesh component on it
 	 * @InAnimSequenceFilename File to create
+	 * @bCreateLink If true will create a link between the animation sequence and the level sequence
 	 */
 	UFUNCTION(BlueprintCallable, Category = "Editor Scripting | Sequencer Tools | Animation")
-	static bool ExportAnimSequence(UWorld* World, ULevelSequence*  Sequence, UAnimSequence* AnimSequence, UAnimSeqExportOption* ExportOption, const FSequencerBindingProxy& Binding);
+	static bool ExportAnimSequence(UWorld* World, ULevelSequence*  Sequence, UAnimSequence* AnimSequence, UAnimSeqExportOption* ExportOption, const FSequencerBindingProxy& Binding, bool bCreateLink);
+
+	/*
+	 * Get the link to the level sequence if it exists on this anim sequence
+	 *
+	 * @InAnimSequence AnimSequence to get links from
+	 * @return Returns the link object if it exists, nullptr if it doesn't
+	*/
+	UFUNCTION(BlueprintCallable, Category = "Editor Scripting | Sequencer Tools | Animation")
+	UAnimSequenceLevelSequenceLink* GetLevelSequenceLinkFromAnimSequence(UAnimSequence* InAnimSequence);
+
+	/*
+	 * Get the links to the anim sequences if they exist on this level sequence
+	 *
+	 * @InLevelSequence LevelSequence to get links from
+	 * @return Returns the link object if it exists, nullptr if it doesn't
+	*/
+	UFUNCTION(BlueprintCallable, Category = "Editor Scripting | Sequencer Tools | Animation")
+	ULevelSequenceAnimSequenceLink* GetAnimSequenceLinkFromLevelSequence(ULevelSequence* InLevelSequence);
 
 	/*
 	 * Import FBX onto Passed in Bindings
diff --git a/Engine/Source/Editor/MovieSceneTools/Private/MovieSceneToolsModule.cpp b/Engine/Source/Editor/MovieSceneTools/Private/MovieSceneToolsModule.cpp
--- a/Engine/Source/Editor/MovieSceneTools/Private/MovieSceneToolsModule.cpp
+++ b/Engine/Source/Editor/MovieSceneTools/Private/MovieSceneToolsModule.cpp
@@ -82,6 +82,10 @@
 #include "EditorModeManager.h"
 #include "EditModes/SkeletalAnimationTrackEditMode.h"
 
+#include "LevelSequence.h"
+#include "LevelSequenceAnimSequenceLink.h"
+#include "AnimSequenceLevelSequenceLink.h"
+#include "EditorAnimUtils.h"
 
 #define LOCTEXT_NAMESPACE "FMovieSceneToolsModule"
 
@@ -186,6 +190,7 @@ void FMovieSceneToolsModule::StartupModule()
 	UMovieSceneEventSectionBase::UpgradeLegacyEventEndpoint.BindStatic(UpgradeLegacyEventEndpointForSection);
 	UMovieSceneEventSectionBase::PostDuplicateSectionEvent.BindStatic(PostDuplicateEventSection);
 	UMovieSceneEventSectionBase::RemoveForCookEvent.BindStatic(RemoveForCookEventSection);
+	ULevelSequence::PostDuplicateEvent.BindStatic(PostDuplicateEvent);
 
 	auto OnObjectsReplaced = [](const TMap<UObject*, UObject*>& ReplacedObjects)
 	{
@@ -219,6 +224,7 @@ void FMovieSceneToolsModule::ShutdownModule()
 	UMovieSceneEventSectionBase::UpgradeLegacyEventEndpoint = UMovieSceneEventSectionBase::FUpgradeLegacyEventEndpoint();
 	UMovieSceneEventSectionBase::PostDuplicateSectionEvent = UMovieSceneEventSectionBase::FPostDuplicateEvent();
 	UMovieSceneEventSectionBase::RemoveForCookEvent = UMovieSceneEventSectionBase::FRemoveForCookEvent();
+	ULevelSequence::PostDuplicateEvent = ULevelSequence::FPostDuplicateEvent();
 
 	if (ICurveEditorModule* CurveEditorModule = FModuleManager::GetModulePtr<ICurveEditorModule>("CurveEditor"))
 	{
@@ -318,6 +324,49 @@ void FMovieSceneToolsModule::RemoveForCookEventSection(UMovieSceneEventSectionBa
 	}
 }
 
+//When we duplicate a ULevelSequence we check to see if there are any linked UAnimSequences in the asset user data,
+//if so we make a copy of the anim sequence, and then set the links to the copies.
+void FMovieSceneToolsModule::PostDuplicateEvent(ULevelSequence* LevelSequence)
+{
+	if (LevelSequence && LevelSequence->GetClass()->ImplementsInterface(UInterface_AssetUserData::StaticClass()))
+	{
+		if (IInterface_AssetUserData* AssetUserDataInterface = Cast< IInterface_AssetUserData >(LevelSequence))
+		{
+			ULevelSequenceAnimSequenceLink* LevelAnimLink = AssetUserDataInterface->GetAssetUserData< ULevelSequenceAnimSequenceLink >();
+			if (LevelAnimLink)
+			{
+				for (FLevelSequenceAnimSequenceLinkItem& Item : LevelAnimLink->AnimSequenceLinks)
+				{
+					if (UAnimSequence* AnimSequence = Item.ResolveAnimSequence())
+					{
+						TArray<UAnimSequence*> AnimSequencesToDuplicate;
+						AnimSequencesToDuplicate.Add(AnimSequence);
+						UPackage* DestinationPackage = AnimSequence->GetPackage();
+						EditorAnimUtils::FNameDuplicationRule NameRule;
+						NameRule.FolderPath = FPackageName::GetLongPackagePath(AnimSequence->GetPathName()) / TEXT("");
+						TMap<UAnimSequence*, UAnimSequence*> DuplicatedAnimAssets = EditorAnimUtils::DuplicateAssets<UAnimSequence>(AnimSequencesToDuplicate, DestinationPackage, &NameRule);
+						for (TPair<UAnimSequence*, UAnimSequence*>& Duplicates : DuplicatedAnimAssets)
+						{
+							if (UAnimSequence* NewAnimSequence = Duplicates.Value)
+							{
+								Item.PathToAnimSequence = FSoftObjectPath(NewAnimSequence);
+								if (IInterface_AssetUserData* AnimAssetUserData = Cast< IInterface_AssetUserData >(NewAnimSequence))
+								{
+									UAnimSequenceLevelSequenceLink* AnimLevelLink = AnimAssetUserData->GetAssetUserData< UAnimSequenceLevelSequenceLink >();
+									if (AnimLevelLink)
+									{
+										AnimLevelLink->SetLevelSequence(LevelSequence);
+									}
+								}
+							}
+						}
+					}
+				}
+			}
+		}
+	}
+}
+
 bool FMovieSceneToolsModule::UpgradeLegacyEventEndpointForSection(UMovieSceneEventSectionBase* Section)
 {
 	UMovieSceneSequence*       Sequence           = Section->GetTypedOuter<UMovieSceneSequence>();
diff --git a/Engine/Source/Editor/MovieSceneTools/Public/MovieSceneToolsModule.h b/Engine/Source/Editor/MovieSceneTools/Public/MovieSceneToolsModule.h
--- a/Engine/Source/Editor/MovieSceneTools/Public/MovieSceneToolsModule.h
+++ b/Engine/Source/Editor/MovieSceneTools/Public/MovieSceneToolsModule.h
@@ -18,6 +18,7 @@ class UMovieScene;
 class UMovieSceneSection;
 class UMovieSceneEventSectionBase;
 class IMovieSceneToolsTrackImporter;
+class ULevelSequence;
 
 class IMovieSceneToolsTakeData
 {
@@ -78,6 +79,7 @@ private:
 	static bool UpgradeLegacyEventEndpointForSection(UMovieSceneEventSectionBase* Section);
 	static void PostDuplicateEventSection(UMovieSceneEventSectionBase* Section);
 	static void RemoveForCookEventSection(UMovieSceneEventSectionBase* Section);
+	static void PostDuplicateEvent(ULevelSequence* LevelSequence);
 
 private:
 
diff --git a/Engine/Source/Runtime/LevelSequence/Private/LevelSequence.cpp b/Engine/Source/Runtime/LevelSequence/Private/LevelSequence.cpp
--- a/Engine/Source/Runtime/LevelSequence/Private/LevelSequence.cpp
+++ b/Engine/Source/Runtime/LevelSequence/Private/LevelSequence.cpp
@@ -39,6 +39,9 @@
 #if WITH_EDITOR
 	#include "UObject/SequencerObjectVersion.h"
 	#include "UObject/ObjectRedirector.h"
+
+ULevelSequence::FPostDuplicateEvent ULevelSequence::PostDuplicateEvent;
+
 #endif
 
 static TAutoConsoleVariable<int32> CVarDefaultLockEngineToDisplayRate(
@@ -236,6 +239,13 @@ void ULevelSequence::PostDuplicate(bool bDuplicateForPIE)
 		DirectorClass = nullptr;
 	}
 #endif
+
+#if WITH_EDITOR
+	if (PostDuplicateEvent.IsBound())
+	{
+		PostDuplicateEvent.Execute(this);
+	}
+#endif
 }
 
 void ULevelSequence::PostLoad()
diff --git a/Engine/Source/Runtime/LevelSequence/Public/AnimSequenceLevelSequenceLink.h b/Engine/Source/Runtime/LevelSequence/Public/AnimSequenceLevelSequenceLink.h
--- a/Engine/Source/Runtime/LevelSequence/Public/AnimSequenceLevelSequenceLink.h
+++ b/Engine/Source/Runtime/LevelSequence/Public/AnimSequenceLevelSequenceLink.h
@@ -7,16 +7,16 @@
 class ULevelSequence;
 
 /** Link To Level Sequence That may be driving the anim sequence*/
-UCLASS()
+UCLASS(BlueprintType)
 class LEVELSEQUENCE_API UAnimSequenceLevelSequenceLink : public UAssetUserData
 {
 	GENERATED_UCLASS_BODY()
 
 public:
-	UPROPERTY()
+	UPROPERTY(BlueprintReadWrite, Category = Property)
 	FGuid SkelTrackGuid;
 
-	UPROPERTY()
+	UPROPERTY(BlueprintReadWrite, Category = Property)
 	FSoftObjectPath PathToLevelSequence;
 
 	void SetLevelSequence(ULevelSequence* InLevelSequence);
diff --git a/Engine/Source/Runtime/LevelSequence/Public/LevelSequence.h b/Engine/Source/Runtime/LevelSequence/Public/LevelSequence.h
--- a/Engine/Source/Runtime/LevelSequence/Public/LevelSequence.h
+++ b/Engine/Source/Runtime/LevelSequence/Public/LevelSequence.h
@@ -75,7 +75,11 @@ public:
 	virtual ETrackSupport IsTrackSupported(TSubclassOf<class UMovieSceneTrack> InTrackClass) const override;
 	virtual void GetAssetRegistryTagMetadata(TMap<FName, FAssetRegistryTagMetadata>& OutMetadata) const override;
 	virtual void GetAssetRegistryTags(TArray<FAssetRegistryTag>& OutTags) const override;
+	
+	DECLARE_DELEGATE_RetVal_OneParam(void, FPostDuplicateEvent, ULevelSequence*);
+	static FPostDuplicateEvent PostDuplicateEvent;
 #endif
+
 	virtual void PostDuplicate(bool bDuplicateForPIE) override;
 
 	void LocateBoundObjects(const FGuid& ObjectId, UObject* Context, FName StreamedLevelAssetPath, TArray<UObject*, TInlineAllocator<1>>& OutObjects) const;
diff --git a/Engine/Source/Runtime/LevelSequence/Public/LevelSequenceAnimSequenceLink.h b/Engine/Source/Runtime/LevelSequence/Public/LevelSequenceAnimSequenceLink.h
--- a/Engine/Source/Runtime/LevelSequence/Public/LevelSequenceAnimSequenceLink.h
+++ b/Engine/Source/Runtime/LevelSequence/Public/LevelSequenceAnimSequenceLink.h
@@ -8,28 +8,28 @@
 class UAnimSequence;
 
 /** Link To Anim Sequence that we are linked too.*/
-USTRUCT()
+USTRUCT(BlueprintType)
 struct LEVELSEQUENCE_API FLevelSequenceAnimSequenceLinkItem
 {
 	GENERATED_BODY()
 
 public:
-	UPROPERTY()
+	UPROPERTY(BlueprintReadWrite, Category = Property)
 	FGuid SkelTrackGuid;
 
-	UPROPERTY()
+	UPROPERTY(BlueprintReadWrite, Category = Property)
 	FSoftObjectPath PathToAnimSequence;
 
 	//From Editor Only UAnimSeqExportOption we cache this since we can re-import dynamically
-	UPROPERTY()
+	UPROPERTY(BlueprintReadWrite, Category = Property)
 	bool bExportTransforms = true;
-	UPROPERTY()
+	UPROPERTY(BlueprintReadWrite, Category = Property)
 	bool bExportMorphTargets = true;
-	UPROPERTY()
+	UPROPERTY(BlueprintReadWrite, Category = Property)
 	bool bExportAttributeCurves = true;
-	UPROPERTY()
+	UPROPERTY(BlueprintReadWrite, Category = Property)
 	bool bExportMaterialCurves = true;
-	UPROPERTY()
+	UPROPERTY(BlueprintReadWrite, Category = Property)
 	bool bRecordInWorldSpace = false;
 
 	void SetAnimSequence(UAnimSequence* InAnimSequence);
@@ -38,12 +38,12 @@ public:
 };
 
 /** Link To Set of Anim Sequences that we may be linked to.*/
-UCLASS()
+UCLASS(BlueprintType)
 class LEVELSEQUENCE_API ULevelSequenceAnimSequenceLink : public UAssetUserData
 {
 	GENERATED_UCLASS_BODY()
 
 public:
-	UPROPERTY()
+	UPROPERTY(BlueprintReadWrite, Category = Links)
 	TArray< FLevelSequenceAnimSequenceLinkItem> AnimSequenceLinks;
 };
\ No newline at end of file
