Sequencer: Linking Bp/Python forgot static
#trivial

diff --git a/Engine/Plugins/MovieScene/SequencerScripting/Source/SequencerScriptingEditor/Public/SequencerTools.h b/Engine/Plugins/MovieScene/SequencerScripting/Source/SequencerScriptingEditor/Public/SequencerTools.h
--- a/Engine/Plugins/MovieScene/SequencerScripting/Source/SequencerScriptingEditor/Public/SequencerTools.h
+++ b/Engine/Plugins/MovieScene/SequencerScripting/Source/SequencerScriptingEditor/Public/SequencerTools.h
@@ -153,7 +153,7 @@ public:
 	 * @return Returns the link object if it exists, nullptr if it doesn't
 	*/
 	UFUNCTION(BlueprintCallable, Category = "Editor Scripting | Sequencer Tools | Animation")
-	UAnimSequenceLevelSequenceLink* GetLevelSequenceLinkFromAnimSequence(UAnimSequence* InAnimSequence);
+	static UAnimSequenceLevelSequenceLink* GetLevelSequenceLinkFromAnimSequence(UAnimSequence* InAnimSequence);
 
 	/*
 	 * Get the links to the anim sequences if they exist on this level sequence
@@ -162,7 +162,7 @@ public:
 	 * @return Returns the link object if it exists, nullptr if it doesn't
 	*/
 	UFUNCTION(BlueprintCallable, Category = "Editor Scripting | Sequencer Tools | Animation")
-	ULevelSequenceAnimSequenceLink* GetAnimSequenceLinkFromLevelSequence(ULevelSequence* InLevelSequence);
+	static ULevelSequenceAnimSequenceLink* GetAnimSequenceLinkFromLevelSequence(ULevelSequence* InLevelSequence);
 
 	/*
 	 * Import FBX onto Passed in Bindings
