Sequencer: Exposed UMovieSceneSection::GetAutoSizeRange() to scripting via UMovieSceneSectionExtensions. For tracks that support auto size (animation, sound, etc.) you can now query if the autosize data has a start/end frame, as well as get it. This simplifies adding assets via Python as you can use a unified function (for all tracks that support GetAutoSizeRange) regardless of underlying asset type which simplifies the need to have a unique way of figuring out how long an asset is for each type.

diff --git a/Engine/Plugins/MovieScene/SequencerScripting/Source/SequencerScripting/Private/ExtensionLibraries/MovieSceneSectionExtensions.cpp b/Engine/Plugins/MovieScene/SequencerScripting/Source/SequencerScripting/Private/ExtensionLibraries/MovieSceneSectionExtensions.cpp
--- a/Engine/Plugins/MovieScene/SequencerScripting/Source/SequencerScripting/Private/ExtensionLibraries/MovieSceneSectionExtensions.cpp
+++ b/Engine/Plugins/MovieScene/SequencerScripting/Source/SequencerScripting/Private/ExtensionLibraries/MovieSceneSectionExtensions.cpp
@@ -155,6 +155,130 @@ float UMovieSceneSectionExtensions::GetEndFrameSeconds(UMovieSceneSection* Secti
 	return -1.f;
 }
 
+bool UMovieSceneSectionExtensions::GetAutoSizeHasStartFrame(UMovieSceneSection* Section)
+{
+	if (!Section)
+	{
+		FFrame::KismetExecutionMessage(TEXT("Cannot call GetAutoSizeHasStartFrame on a null section"), ELogVerbosity::Error);
+		return false;
+	}
+
+	TOptional<TRange<FFrameNumber>> AutoSizeRange = Section->GetAutoSizeRange();
+	return AutoSizeRange.IsSet() && AutoSizeRange->GetLowerBound().IsClosed();
+}
+
+int32 UMovieSceneSectionExtensions::GetAutoSizeStartFrame(UMovieSceneSection* Section)
+{
+	if (!Section)
+	{
+		FFrame::KismetExecutionMessage(TEXT("Cannot call GetAutoSizeStartFrame on a null section"), ELogVerbosity::Error);
+		return -1;
+	}
+
+	TOptional<TRange<FFrameNumber>> AutoSizeRange = Section->GetAutoSizeRange();
+	if (!AutoSizeRange.IsSet() || AutoSizeRange->GetLowerBound().IsOpen())
+	{
+		FFrame::KismetExecutionMessage(TEXT("Section does not have an AutoSize start frame"), ELogVerbosity::Error);
+		return -1;
+	}
+
+	UMovieScene* MovieScene = Section->GetTypedOuter<UMovieScene>();
+	if (MovieScene)
+	{
+		FFrameRate DisplayRate = MovieScene->GetDisplayRate();
+		return ConvertFrameTime(UE::MovieScene::DiscreteInclusiveLower(AutoSizeRange.GetValue()), MovieScene->GetTickResolution(), DisplayRate).FloorToFrame().Value;
+	}
+
+	return -1;
+}
+
+float UMovieSceneSectionExtensions::GetAutoSizeStartFrameSeconds(UMovieSceneSection* Section)
+{
+	if (!Section)
+	{
+		FFrame::KismetExecutionMessage(TEXT("Cannot call GetAutoSizeStartFrameSeconds on a null section"), ELogVerbosity::Error);
+		return -1.f;
+	}
+
+	TOptional<TRange<FFrameNumber>> AutoSizeRange = Section->GetAutoSizeRange();
+	if (!AutoSizeRange.IsSet() || AutoSizeRange->GetLowerBound().IsOpen())
+	{
+		FFrame::KismetExecutionMessage(TEXT("Section does not have an AutoSize start frame"), ELogVerbosity::Error);
+		return -1.f;
+	}
+
+	UMovieScene* MovieScene = Section->GetTypedOuter<UMovieScene>();
+	if (MovieScene)
+	{
+		FFrameRate DisplayRate = MovieScene->GetDisplayRate();
+		return DisplayRate.AsSeconds(ConvertFrameTime(UE::MovieScene::DiscreteInclusiveLower(AutoSizeRange.GetValue()), MovieScene->GetTickResolution(), DisplayRate));
+	}
+
+	return -1.f;
+}
+
+bool UMovieSceneSectionExtensions::GetAutoSizeHasEndFrame(UMovieSceneSection* Section)
+{
+	if (!Section)
+	{
+		FFrame::KismetExecutionMessage(TEXT("Cannot call GetAutoSizeHasEndFrame on a null section"), ELogVerbosity::Error);
+		return false;
+	}
+
+	TOptional<TRange<FFrameNumber>> AutoSizeRange = Section->GetAutoSizeRange();
+	return AutoSizeRange.IsSet() && AutoSizeRange->GetUpperBound().IsClosed();
+}
+
+int32 UMovieSceneSectionExtensions::GetAutoSizeEndFrame(UMovieSceneSection* Section)
+{
+	if (!Section)
+	{
+		FFrame::KismetExecutionMessage(TEXT("Cannot call GetAutoSizeEndFrame on a null section"), ELogVerbosity::Error);
+		return -1;
+	}
+
+	TOptional<TRange<FFrameNumber>> AutoSizeRange = Section->GetAutoSizeRange();
+	if (!AutoSizeRange.IsSet() || AutoSizeRange->GetUpperBound().IsOpen())
+	{
+		FFrame::KismetExecutionMessage(TEXT("Section does not have an AutoSize end frame"), ELogVerbosity::Error);
+		return -1;
+	}
+
+	UMovieScene* MovieScene = Section->GetTypedOuter<UMovieScene>();
+	if (MovieScene)
+	{
+		FFrameRate DisplayRate = MovieScene->GetDisplayRate();
+		return ConvertFrameTime(UE::MovieScene::DiscreteExclusiveUpper(AutoSizeRange.GetValue()), MovieScene->GetTickResolution(), DisplayRate).FloorToFrame().Value;
+	}
+
+	return -1;
+}
+
+float UMovieSceneSectionExtensions::GetAutoSizeEndFrameSeconds(UMovieSceneSection* Section)
+{
+	if (!Section)
+	{
+		FFrame::KismetExecutionMessage(TEXT("Cannot call GetAutoSizeEndFrameSeconds on a null section"), ELogVerbosity::Error);
+		return -1.f;
+	}
+
+	TOptional<TRange<FFrameNumber>> AutoSizeRange = Section->GetAutoSizeRange();
+	if (!AutoSizeRange.IsSet() || AutoSizeRange->GetUpperBound().IsOpen())
+	{
+		FFrame::KismetExecutionMessage(TEXT("Section does not have an AutoSize end frame"), ELogVerbosity::Error);
+		return -1.f;
+	}
+
+	UMovieScene* MovieScene = Section->GetTypedOuter<UMovieScene>();
+	if (MovieScene)
+	{
+		FFrameRate DisplayRate = MovieScene->GetDisplayRate();
+		return DisplayRate.AsSeconds(ConvertFrameTime(UE::MovieScene::DiscreteExclusiveUpper(AutoSizeRange.GetValue()), MovieScene->GetTickResolution(), DisplayRate));
+	}
+
+	return -1.f;
+}
+
 void UMovieSceneSectionExtensions::SetRange(UMovieSceneSection* Section, int32 StartFrame, int32 EndFrame)
 {
 	if (!Section)
diff --git a/Engine/Plugins/MovieScene/SequencerScripting/Source/SequencerScripting/Public/ExtensionLibraries/MovieSceneSectionExtensions.h b/Engine/Plugins/MovieScene/SequencerScripting/Source/SequencerScripting/Public/ExtensionLibraries/MovieSceneSectionExtensions.h
--- a/Engine/Plugins/MovieScene/SequencerScripting/Source/SequencerScripting/Public/ExtensionLibraries/MovieSceneSectionExtensions.h
+++ b/Engine/Plugins/MovieScene/SequencerScripting/Source/SequencerScripting/Public/ExtensionLibraries/MovieSceneSectionExtensions.h
@@ -34,7 +34,7 @@ public:
 	static bool HasStartFrame(UMovieSceneSection* Section);
 
 	/**
-	 * Get start frame
+	 * Get start frame. Will throw an exception if section has no start frame, use HasStartFrame to check first.
 	 *
 	 * @param Section        The section within which to get the start frame
 	 * @return Start frame of this section
@@ -43,7 +43,7 @@ public:
 	static int32 GetStartFrame(UMovieSceneSection* Section);
 
 	/**
-	 * Get start time in seconds
+	 * Get start time in seconds. Will throw an exception if section has no start frame, use HasStartFrame to check first.
 	 *
 	 * @param Section        The section within which to get the start time
 	 * @return Start time of this section
@@ -61,7 +61,7 @@ public:
 	static bool HasEndFrame(UMovieSceneSection* Section);
 
 	/**
-	 * Get end frame
+	 * Get end frame. Will throw an exception if section has no end frame, use HasEndFrame to check first.
 	 *
 	 * @param Section        The section within which to get the end frame
 	 * @return End frame of this section
@@ -70,7 +70,7 @@ public:
 	static int32 GetEndFrame(UMovieSceneSection* Section);
 
 	/**
-	 * Get end time in seconds
+	 * Get end time in seconds. Will throw an exception if section has no end frame, use HasEndFrame to check first.
 	 *
 	 * @param Section        The section within which to get the end time
 	 * @return End time of this section
@@ -78,6 +78,60 @@ public:
 	UFUNCTION(BlueprintCallable, Category = "Sequencer|Section", meta = (ScriptMethod))
 	static float GetEndFrameSeconds(UMovieSceneSection* Section);
 
+	/**
+	 * Checks to see if this section has an AutoSize implementation, and if so, if that implementation has a start frame.
+	 *
+	 * @param Section        The section being queried
+	 * @return Whether this section has a valid autosize range, and a valid start frame
+	 */
+	UFUNCTION(BlueprintCallable, Category = "Sequencer|Section", meta = (ScriptMethod))
+	static bool GetAutoSizeHasStartFrame(UMovieSceneSection* Section);
+
+	/**
+	 * Get start frame of the AutoSize. Will throw an exception if section has no start frame, use GetAutoSizeHasStartFrame to check first.
+	 *
+	 * @param Section        The section being queried
+	 * @return The start frame of the AutoSize data.
+	 */
+	UFUNCTION(BlueprintCallable, Category = "Sequencer|Section", meta = (ScriptMethod))
+	static int32 GetAutoSizeStartFrame(UMovieSceneSection* Section);
+
+	/**
+	 * Get start time of the AutoSize in seconds. Will throw an exception if section has no start frame, use GetAutoSizeHasStartFrame to check first.
+	 *
+	 * @param Section        The section being queried
+	 * @return The start frame of the AutoSize data in seconds.
+	 */
+	UFUNCTION(BlueprintCallable, Category = "Sequencer|Section", meta = (ScriptMethod))
+	static float GetAutoSizeStartFrameSeconds(UMovieSceneSection* Section);
+
+	/**
+	 * Checks to see if this section has an AutoSize implementation, and if so, if that implementation has a end frame.
+	 *
+	 * @param Section        The section being queried
+	 * @return Whether this section has a valid autosize range, and a valid end frame
+	 */
+	UFUNCTION(BlueprintCallable, Category = "Sequencer|Section", meta = (ScriptMethod))
+	static bool GetAutoSizeHasEndFrame(UMovieSceneSection* Section);
+
+	/**
+	 * Get end frame of the AutoSize. Will throw an exception if section has no end frame, use GetAutoSizeHasEndFrame to check first.
+	 *
+	 * @param Section        The section being queried
+	 * @return The end frame of the AutoSize data.
+	 */
+	UFUNCTION(BlueprintCallable, Category = "Sequencer|Section", meta = (ScriptMethod))
+	static int32 GetAutoSizeEndFrame(UMovieSceneSection* Section);
+
+	/**
+	 * Get end time of the AutoSize seconds. Will throw an exception if section has no end frame, use GetAutoSizeHasEndFrame to check first.
+	 *
+	 * @param Section        The section being queried
+	 * @return The end frame of the AutoSize data in seconds.
+	 */
+	UFUNCTION(BlueprintCallable, Category = "Sequencer|Section", meta = (ScriptMethod))
+	static float GetAutoSizeEndFrameSeconds(UMovieSceneSection* Section);
+
 	/**
 	 * Set range
 	 *
