Expose LinkAnimSequence function to scripting

diff --git a/Engine/Plugins/MovieScene/SequencerScripting/Source/SequencerScriptingEditor/Private/SequencerTools.cpp b/Engine/Plugins/MovieScene/SequencerScripting/Source/SequencerScriptingEditor/Private/SequencerTools.cpp
--- a/Engine/Plugins/MovieScene/SequencerScripting/Source/SequencerScriptingEditor/Private/SequencerTools.cpp
+++ b/Engine/Plugins/MovieScene/SequencerScripting/Source/SequencerScriptingEditor/Private/SequencerTools.cpp
@@ -330,7 +330,23 @@ bool USequencerToolsFunctionLibrary::ExportAnimSequence(UWorld* World, ULevelSeq
 	//create the link to the anim sequence
 	if (bResult && bCreateLink)
 	{
-		//create from anim sequence back to level sequence
+		return LinkAnimSequence(Sequence, AnimSequence, ExportOptions, Binding);
+	}
+	return bResult;
+}
+
+bool USequencerToolsFunctionLibrary::LinkAnimSequence(ULevelSequence*  Sequence,  UAnimSequence* AnimSequence, const UAnimSeqExportOption* ExportOptions,const FSequencerBindingProxy& Binding)
+{
+	if (!Sequence || !AnimSequence || !ExportOptions || Binding.Sequence != Sequence)
+	{
+		return false;
+	}
+
+	if (Sequence && Sequence->GetClass()->ImplementsInterface(UInterface_AssetUserData::StaticClass())
+		&& AnimSequence->GetClass()->ImplementsInterface(UInterface_AssetUserData::StaticClass()))
+	{
+		Sequence->Modify();
+		AnimSequence->Modify();
 		if (IInterface_AssetUserData* AnimAssetUserData = Cast< IInterface_AssetUserData >(AnimSequence))
 		{
 			UAnimSequenceLevelSequenceLink* AnimLevelLink = AnimAssetUserData->GetAssetUserData< UAnimSequenceLevelSequenceLink >();
@@ -339,11 +355,10 @@ bool USequencerToolsFunctionLibrary::ExportAnimSequence(UWorld* World, ULevelSeq
 				AnimLevelLink = NewObject<UAnimSequenceLevelSequenceLink>(AnimSequence, NAME_None, RF_Public | RF_Transactional);
 				AnimAssetUserData->AddAssetUserData(AnimLevelLink);
 			}
-
+			
 			AnimLevelLink->SetLevelSequence(Sequence);
 			AnimLevelLink->SkelTrackGuid = Binding.BindingID;
 		}
-		//create from level sequence to anim sequence, trickier since we can have multiples here.
 		if (IInterface_AssetUserData* AssetUserDataInterface = Cast< IInterface_AssetUserData >(Sequence))
 		{
 			bool bAddItem = true;
@@ -356,7 +371,7 @@ bool USequencerToolsFunctionLibrary::ExportAnimSequence(UWorld* World, ULevelSeq
 					{
 						bAddItem = false;
 						UAnimSequence* OtherAnimSequence = LevelAnimLinkItem.ResolveAnimSequence();
-
+						
 						if (OtherAnimSequence != AnimSequence)
 						{
 							if (IInterface_AssetUserData* OtherAnimAssetUserData = Cast< IInterface_AssetUserData >(OtherAnimSequence))
@@ -364,6 +379,7 @@ bool USequencerToolsFunctionLibrary::ExportAnimSequence(UWorld* World, ULevelSeq
 								UAnimSequenceLevelSequenceLink* OtherAnimLevelLink = OtherAnimAssetUserData->GetAssetUserData< UAnimSequenceLevelSequenceLink >();
 								if (OtherAnimLevelLink)
 								{
+									OtherAnimSequence->Modify();
 									OtherAnimAssetUserData->RemoveUserDataOfClass(UAnimSequenceLevelSequenceLink::StaticClass());
 								}
 							}
@@ -382,7 +398,6 @@ bool USequencerToolsFunctionLibrary::ExportAnimSequence(UWorld* World, ULevelSeq
 			else
 			{
 				LevelAnimLink = NewObject<ULevelSequenceAnimSequenceLink>(Sequence, NAME_None, RF_Public | RF_Transactional);
-
 			}
 			if (bAddItem == true)
 			{
@@ -400,7 +415,7 @@ bool USequencerToolsFunctionLibrary::ExportAnimSequence(UWorld* World, ULevelSeq
 			}
 		}
 	}
-	return bResult;
+	return true;
 }
 
 UAnimSequenceLevelSequenceLink* USequencerToolsFunctionLibrary::GetLevelSequenceLinkFromAnimSequence(UAnimSequence* InAnimSequence)
diff --git a/Engine/Plugins/MovieScene/SequencerScripting/Source/SequencerScriptingEditor/Public/SequencerTools.h b/Engine/Plugins/MovieScene/SequencerScripting/Source/SequencerScriptingEditor/Public/SequencerTools.h
--- a/Engine/Plugins/MovieScene/SequencerScripting/Source/SequencerScriptingEditor/Public/SequencerTools.h
+++ b/Engine/Plugins/MovieScene/SequencerScripting/Source/SequencerScriptingEditor/Public/SequencerTools.h
@@ -146,6 +146,17 @@ public:
 	UFUNCTION(BlueprintCallable, Category = "Editor Scripting | Sequencer Tools | Animation")
 	static bool ExportAnimSequence(UWorld* World, ULevelSequence*  Sequence, UAnimSequence* AnimSequence, UAnimSeqExportOption* ExportOption, const FSequencerBindingProxy& Binding, bool bCreateLink);
 
+	/*
+	 * Links a LevelSequence's SkeletalMesh binding to an existing anim sequence.
+	 *
+	 * @InSequence Sequence to link from
+	 * @AnimSequence The AnimSequence to link to.
+	 * @ExportOption The export options that should be used when baking the LevelSequence.
+	 * @InBinding Binding that has a skelmesh component on it
+	 */
+	UFUNCTION(BlueprintCallable, Category = "Editor Scripting | Sequencer Tools | Animation")
+	static bool LinkAnimSequence(ULevelSequence* Sequence, UAnimSequence* AnimSequence, const UAnimSeqExportOption* ExportOptions, const FSequencerBindingProxy& Binding);
+
 	/*
 	 * Get the link to the level sequence if it exists on this anim sequence
 	 *
