PR #9302: Update wrong Blueprint DisplayName of "SetValue" node. (Contributed by [REDACTED_PERSON])

diff --git a/Engine/Plugins/MovieScene/SequencerScripting/Source/SequencerScripting/Private/KeysAndChannels/MovieSceneScriptingEvent.h b/Engine/Plugins/MovieScene/SequencerScripting/Source/SequencerScripting/Private/KeysAndChannels/MovieSceneScriptingEvent.h
--- a/Engine/Plugins/MovieScene/SequencerScripting/Source/SequencerScripting/Private/KeysAndChannels/MovieSceneScriptingEvent.h
+++ b/Engine/Plugins/MovieScene/SequencerScripting/Source/SequencerScripting/Private/KeysAndChannels/MovieSceneScriptingEvent.h
@@ -55,7 +55,7 @@ public:
 	* Sets the value for this key, reflecting it in the owning channel.
 	* @param InNewValue	The new event payload for this key.
 	*/
-	UFUNCTION(BlueprintCallable, Category = "Sequencer|Keys", meta = (DisplayName = "Get Value (Event)"))
+	UFUNCTION(BlueprintCallable, Category = "Sequencer|Keys", meta = (DisplayName = "Set Value (Event)"))
 	void SetValue(const FMovieSceneEvent& InNewValue)
 	{
 		SetValueInChannel(KeyHandle, InNewValue);
