Sequencer: Move SequencerBindingProxy to MovieScene module and rename as MovieSceneBindingProxy

Deprecate GetSelectedObjects()/SelectObjects in favor of GetSelectedBindings/SelectBindings which operates on FMovieSceneBindingProxy for consistency with all other blueprint/python functions
Deprecate SequencerBindingProxy and typedef it to MovieSceneBindingProxy

Tested that calls with SequencerBindingProxy produce warnings but still compile
Tested that blueprints with SequencerBindingProxy are automatically redirected to use MovieSceneBindingProxy

diff --git a/Engine/Config/BaseEngine.ini b/Engine/Config/BaseEngine.ini
--- a/Engine/Config/BaseEngine.ini
+++ b/Engine/Config/BaseEngine.ini
@@ -1079,6 +1079,8 @@ DefaultWorldPartitionHLODLayer=/Engine/Maps/Templates/HLODs/HLODLayer_Instanced.
 
 +ClassRedirects=(OldName="RigVMTemplateNode", NewName="/Script/RigVMDeveloper.RigVMUnitNode")
 
++StructRedirects=(OldName="/Script/SequencerScripting.SequencerBindingProxy", NewName="/Script/MovieScene.MovieSceneBindingProxy")
+
 ; Blueprint real number support
 +FunctionRedirects=(OldName="KismetSystemLibrary.MakeLiteralFloat", NewName="MakeLiteralDouble")
 
diff --git a/Engine/Plugins/Animation/ControlRig/Source/ControlRigEditor/Private/ControlRigSequencerEditorLibrary.cpp b/Engine/Plugins/Animation/ControlRig/Source/ControlRigEditor/Private/ControlRigSequencerEditorLibrary.cpp
--- a/Engine/Plugins/Animation/ControlRig/Source/ControlRigEditor/Private/ControlRigSequencerEditorLibrary.cpp
+++ b/Engine/Plugins/Animation/ControlRig/Source/ControlRigEditor/Private/ControlRigSequencerEditorLibrary.cpp
@@ -82,7 +82,7 @@ TArray<FControlRigSequencerBindingProxy> UControlRigSequencerEditorLibrary::GetC
 	return ControlRigBindingProxies;
 }
 
-static TArray<UObject*> GetBoundObjects(UWorld* World, ULevelSequence* LevelSequence, const FSequencerBindingProxy& Binding, ULevelSequencePlayer** OutPlayer, ALevelSequenceActor** OutActor)
+static TArray<UObject*> GetBoundObjects(UWorld* World, ULevelSequence* LevelSequence, const FMovieSceneBindingProxy& Binding, ULevelSequencePlayer** OutPlayer, ALevelSequenceActor** OutActor)
 {
 	FMovieSceneSequencePlaybackSettings Settings;
 	FLevelSequenceCameraSettings CameraSettings;
@@ -260,7 +260,7 @@ static UMovieSceneControlRigParameterTrack* AddControlRig(ULevelSequence* LevelS
 	return nullptr;
 }
 
-UMovieSceneTrack* UControlRigSequencerEditorLibrary::FindOrCreateControlRigTrack(UWorld* World, ULevelSequence* LevelSequence, const UClass* ControlRigClass, const FSequencerBindingProxy& InBinding)
+UMovieSceneTrack* UControlRigSequencerEditorLibrary::FindOrCreateControlRigTrack(UWorld* World, ULevelSequence* LevelSequence, const UClass* ControlRigClass, const FMovieSceneBindingProxy& InBinding)
 {
 	UMovieScene* MovieScene = InBinding.Sequence ? InBinding.Sequence->GetMovieScene() : nullptr;
 	UMovieSceneTrack* BaseTrack = nullptr;
@@ -321,7 +321,7 @@ UMovieSceneTrack* UControlRigSequencerEditorLibrary::FindOrCreateControlRigTrack
 }
 
 
-TArray<UMovieSceneTrack*> UControlRigSequencerEditorLibrary::FindOrCreateControlRigComponentTrack(UWorld* World, ULevelSequence* LevelSequence, const FSequencerBindingProxy& InBinding)
+TArray<UMovieSceneTrack*> UControlRigSequencerEditorLibrary::FindOrCreateControlRigComponentTrack(UWorld* World, ULevelSequence* LevelSequence, const FMovieSceneBindingProxy& InBinding)
 {
 	TArray< UMovieSceneTrack*> Tracks;
 	TArray<UObject*, TInlineAllocator<1>> Result;
@@ -615,7 +615,7 @@ void UControlRigSequencerEditorLibrary::SetControlRigWorldTransforms(ULevelSeque
 }
 
 bool UControlRigSequencerEditorLibrary::BakeToControlRig(UWorld* World, ULevelSequence* LevelSequence, UClass* InClass, UAnimSeqExportOption* ExportOptions, bool bReduceKeys, float Tolerance,
-	const FSequencerBindingProxy& Binding)
+	const FMovieSceneBindingProxy& Binding)
 {
 	UMovieScene* MovieScene = LevelSequence->GetMovieScene();
 	if (Binding.Sequence != LevelSequence)
diff --git a/Engine/Plugins/Animation/ControlRig/Source/ControlRigEditor/Public/ControlRigSequencerEditorLibrary.h b/Engine/Plugins/Animation/ControlRig/Source/ControlRigEditor/Public/ControlRigSequencerEditorLibrary.h
--- a/Engine/Plugins/Animation/ControlRig/Source/ControlRigEditor/Public/ControlRigSequencerEditorLibrary.h
+++ b/Engine/Plugins/Animation/ControlRig/Source/ControlRigEditor/Public/ControlRigSequencerEditorLibrary.h
@@ -9,7 +9,7 @@
 
 #include "CoreMinimal.h"
 #include "UObject/ObjectMacros.h"
-#include "SequencerBindingProxy.h"
+#include "MovieSceneBindingProxy.h"
 #include "ControlRig.h"
 #include "Tools/ControlRigSnapper.h"
 #include "TransformNoScale.h"
@@ -31,14 +31,14 @@ struct FControlRigSequencerBindingProxy
 		, Track(nullptr)
 	{}
 
-	FControlRigSequencerBindingProxy(const FSequencerBindingProxy& InProxy, UControlRig* InControlRig, UMovieSceneControlRigParameterTrack* InTrack)
+	FControlRigSequencerBindingProxy(const FMovieSceneBindingProxy& InProxy, UControlRig* InControlRig, UMovieSceneControlRigParameterTrack* InTrack)
 		: Proxy(InProxy)
 		, ControlRig(InControlRig)
 		, Track(InTrack)
 	{}
 
 	UPROPERTY(BlueprintReadOnly, Category = ControlRig)
-	FSequencerBindingProxy Proxy;
+	FMovieSceneBindingProxy Proxy;
 
 	UPROPERTY(BlueprintReadOnly, Category = ControlRig)
 	TObjectPtr<UControlRig> ControlRig;
@@ -83,7 +83,7 @@ public:
 	* @return returns Return the found or created track
 	*/
 	UFUNCTION(BlueprintCallable, Category = "Editor Scripting | Sequencer Tools | Control Rig")
-	static UMovieSceneTrack* FindOrCreateControlRigTrack(UWorld* World, ULevelSequence* LevelSequence, const UClass* ControlRigClass, const FSequencerBindingProxy& InBinding);
+	static UMovieSceneTrack* FindOrCreateControlRigTrack(UWorld* World, ULevelSequence* LevelSequence, const UClass* ControlRigClass, const FMovieSceneBindingProxy& InBinding);
 
 	/**
 	* Find or create a Control Rig Component
@@ -93,7 +93,7 @@ public:
 	* @return returns Find array of component Control Rigs that were found or created
 	*/
 	UFUNCTION(BlueprintCallable, Category = "Editor Scripting | Sequencer Tools | Control Rig")
-	static TArray<UMovieSceneTrack*> FindOrCreateControlRigComponentTrack(UWorld* World, ULevelSequence* LevelSequence, const FSequencerBindingProxy& InBinding);
+	static TArray<UMovieSceneTrack*> FindOrCreateControlRigComponentTrack(UWorld* World, ULevelSequence* LevelSequence, const FMovieSceneBindingProxy& InBinding);
 	
 	/**
 	* Load anim sequence into this control rig section
@@ -125,7 +125,7 @@ public:
 	*/
 	UFUNCTION(BlueprintCallable, Category = "Editor Scripting | Sequencer Tools | Control Rig")
 	static bool BakeToControlRig(UWorld* World, ULevelSequence* LevelSequence, UClass* ControlRigClass, UAnimSeqExportOption* ExportOptions, bool bReduceKeys, float Tolerance,
-			const FSequencerBindingProxy& Binding);
+			const FMovieSceneBindingProxy& Binding);
 
 	/**
 	* Peform a Tween operation on the current active sequencer time(must be visible).
@@ -888,4 +888,3 @@ public:
 	bool RenameControlRigControlChannels(ULevelSequence* InSequence, UControlRig* InControlRig, const TArray<FName>& InOldControlNames, const TArray<FName>& InNewControlNames);
 
 };
-
diff --git a/Engine/Plugins/MovieScene/LevelSequenceEditor/Source/LevelSequenceEditor/Private/LevelSequenceEditorBlueprintLibrary.cpp b/Engine/Plugins/MovieScene/LevelSequenceEditor/Source/LevelSequenceEditor/Private/LevelSequenceEditorBlueprintLibrary.cpp
--- a/Engine/Plugins/MovieScene/LevelSequenceEditor/Source/LevelSequenceEditor/Private/LevelSequenceEditorBlueprintLibrary.cpp
+++ b/Engine/Plugins/MovieScene/LevelSequenceEditor/Source/LevelSequenceEditor/Private/LevelSequenceEditorBlueprintLibrary.cpp
@@ -201,6 +201,24 @@ TArray<FGuid> ULevelSequenceEditorBlueprintLibrary::GetSelectedObjects()
 	return OutSelectedGuids;
 }
 
+TArray<FMovieSceneBindingProxy> ULevelSequenceEditorBlueprintLibrary::GetSelectedBindings()
+{
+	TArray<FMovieSceneBindingProxy> OutSelectedBindings;
+	if (CurrentSequencer.IsValid())
+	{
+		TArray<FGuid> OutSelectedGuids;
+		CurrentSequencer.Pin()->GetSelectedObjects(OutSelectedGuids);
+
+		UMovieSceneSequence* Sequence = CurrentSequencer.Pin()->GetFocusedMovieSceneSequence();
+
+		for (const FGuid& SelectedGuid : OutSelectedGuids)
+		{
+			OutSelectedBindings.Add(FMovieSceneBindingProxy(SelectedGuid, Sequence));
+		}
+	}
+	return OutSelectedBindings;
+}
+
 void ULevelSequenceEditorBlueprintLibrary::SelectTracks(const TArray<UMovieSceneTrack*>& Tracks)
 {
 	if (CurrentSequencer.IsValid())
@@ -262,6 +280,17 @@ void ULevelSequenceEditorBlueprintLibrary::SelectObjects(TArray<FGuid> ObjectBin
 	}
 }
 
+void ULevelSequenceEditorBlueprintLibrary::SelectBindings(const TArray<FMovieSceneBindingProxy>& ObjectBindings)
+{
+	if (CurrentSequencer.IsValid())
+	{
+		for (const FMovieSceneBindingProxy& ObjectBinding : ObjectBindings)
+		{
+			CurrentSequencer.Pin()->SelectObject(ObjectBinding.BindingID);
+		}
+	}
+}
+
 void ULevelSequenceEditorBlueprintLibrary::EmptySelection()
 {
 	if (CurrentSequencer.IsValid())
diff --git a/Engine/Plugins/MovieScene/LevelSequenceEditor/Source/LevelSequenceEditor/Public/LevelSequenceEditorBlueprintLibrary.h b/Engine/Plugins/MovieScene/LevelSequenceEditor/Source/LevelSequenceEditor/Public/LevelSequenceEditorBlueprintLibrary.h
--- a/Engine/Plugins/MovieScene/LevelSequenceEditor/Source/LevelSequenceEditor/Public/LevelSequenceEditorBlueprintLibrary.h
+++ b/Engine/Plugins/MovieScene/LevelSequenceEditor/Source/LevelSequenceEditor/Public/LevelSequenceEditorBlueprintLibrary.h
@@ -6,6 +6,7 @@
 #include "Kismet/BlueprintFunctionLibrary.h"
 #include "MovieSceneSequencePlayer.h"
 #include "MovieSceneObjectBindingID.h"
+#include "MovieSceneBindingProxy.h"
 #include "LevelSequenceEditorBlueprintLibrary.generated.h"
 
 class ISequencer;
@@ -136,9 +137,14 @@ public:
 	static TArray<UMovieSceneFolder*> GetSelectedFolders();
 
 	/** Gets the currently selected Object Guids*/
-	UFUNCTION(BlueprintPure, Category = "Level Sequence Editor")
+	UE_DEPRECATED(5.1, "GetSelectedObjects is deprecated, please use GetSelectedBindings which returns an array of FMovieSceneBindingProxy")
+	UFUNCTION(BlueprintPure, Category = "Level Sequence Editor", meta = (DeprecatedFunction, DeprecationMessage="GetSelectedObjects is deprecated, please use GetSelectedBindings which returns an array of FMovieSceneBindingProxy"))
 	static TArray<FGuid> GetSelectedObjects();
 
+	/** Gets the currently selected object bindings */
+	UFUNCTION(BlueprintPure, Category = "Level Sequence Editor")
+	static TArray<FMovieSceneBindingProxy> GetSelectedBindings();
+
 	/** Select tracks */
 	UFUNCTION(BlueprintCallable, Category = "Level Sequence Editor")
 	static void SelectTracks(const TArray<UMovieSceneTrack*>& Tracks);
@@ -156,9 +162,14 @@ public:
 	static void SelectFolders(const TArray<UMovieSceneFolder*>& Folders);
 
 	/** Select objects by GUID */
-	UFUNCTION(BlueprintCallable, Category = "Level Sequence Editor")
+	UE_DEPRECATED(5.1, "SelectObjects is deprecated, please use SelectBindings which takes an FMovieSceneBindingProxy")
+	UFUNCTION(BlueprintCallable, Category = "Level Sequence Editor", meta=(DeprecatedFunction, DeprecationMessage="SelectObjects is deprecated, please use SelectBindings which takes an FMovieSceneBindingProxy"))
 	static void SelectObjects(TArray<FGuid> ObjectBinding);
 
+	/** Select bindings */
+	UFUNCTION(BlueprintCallable, Category = "Level Sequence Editor")
+	static void SelectBindings(const TArray<FMovieSceneBindingProxy>& ObjectBindings);
+
 	/** Empties the current selection. */
 	UFUNCTION(BlueprintCallable, Category = "Level Sequence Editor")
 	static void EmptySelection();
diff --git a/Engine/Plugins/MovieScene/SequencerScripting/Source/SequencerScripting/Private/ExtensionLibraries/MovieSceneBindingExtensions.cpp b/Engine/Plugins/MovieScene/SequencerScripting/Source/SequencerScripting/Private/ExtensionLibraries/MovieSceneBindingExtensions.cpp
--- a/Engine/Plugins/MovieScene/SequencerScripting/Source/SequencerScripting/Private/ExtensionLibraries/MovieSceneBindingExtensions.cpp
+++ b/Engine/Plugins/MovieScene/SequencerScripting/Source/SequencerScripting/Private/ExtensionLibraries/MovieSceneBindingExtensions.cpp
@@ -2,12 +2,12 @@
 
 #include "ExtensionLibraries/MovieSceneBindingExtensions.h"
 #include "ExtensionLibraries/MovieSceneSequenceExtensions.h"
-#include "SequencerBindingProxy.h"
+#include "MovieSceneBindingProxy.h"
 #include "MovieSceneSequence.h"
 #include "MovieScene.h"
 #include "Algo/Find.h"
 
-bool UMovieSceneBindingExtensions::IsValid(const FSequencerBindingProxy& InBinding)
+bool UMovieSceneBindingExtensions::IsValid(const FMovieSceneBindingProxy& InBinding)
 {
 	UMovieScene* MovieScene = InBinding.Sequence ? InBinding.Sequence->GetMovieScene() : nullptr;
 	if (MovieScene && InBinding.BindingID.IsValid())
@@ -18,12 +18,12 @@
 	return false;
 }
 
-FGuid UMovieSceneBindingExtensions::GetId(const FSequencerBindingProxy& InBinding)
+FGuid UMovieSceneBindingExtensions::GetId(const FMovieSceneBindingProxy& InBinding)
 {
 	return InBinding.BindingID;
 }
 
-FText UMovieSceneBindingExtensions::GetDisplayName(const FSequencerBindingProxy& InBinding)
+FText UMovieSceneBindingExtensions::GetDisplayName(const FMovieSceneBindingProxy& InBinding)
 {
 	UMovieScene* MovieScene = InBinding.Sequence ? InBinding.Sequence->GetMovieScene() : nullptr;
 	if (MovieScene && InBinding.BindingID.IsValid())
@@ -34,7 +34,7 @@
 	return FText();
 }
 
-void UMovieSceneBindingExtensions::SetDisplayName(const FSequencerBindingProxy& InBinding, const FText& InDisplayName)
+void UMovieSceneBindingExtensions::SetDisplayName(const FMovieSceneBindingProxy& InBinding, const FText& InDisplayName)
 {
 	UMovieScene* MovieScene = InBinding.Sequence ? InBinding.Sequence->GetMovieScene() : nullptr;
 	if (MovieScene && InBinding.BindingID.IsValid())
@@ -46,7 +46,7 @@
 	}
 }
 
-FString UMovieSceneBindingExtensions::GetName(const FSequencerBindingProxy& InBinding)
+FString UMovieSceneBindingExtensions::GetName(const FMovieSceneBindingProxy& InBinding)
 {
 	UMovieScene* MovieScene = InBinding.Sequence ? InBinding.Sequence->GetMovieScene() : nullptr;
 	if (MovieScene && InBinding.BindingID.IsValid())
@@ -67,7 +67,7 @@
 	return FString();
 }
 
-void UMovieSceneBindingExtensions::SetName(const FSequencerBindingProxy& InBinding, const FString& InName)
+void UMovieSceneBindingExtensions::SetName(const FMovieSceneBindingProxy& InBinding, const FString& InName)
 {
 	UMovieScene* MovieScene = InBinding.Sequence ? InBinding.Sequence->GetMovieScene() : nullptr;
 	if (MovieScene && InBinding.BindingID.IsValid())
@@ -88,7 +88,7 @@
 	}
 }
 
-TArray<UMovieSceneTrack*> UMovieSceneBindingExtensions::GetTracks(const FSequencerBindingProxy& InBinding)
+TArray<UMovieSceneTrack*> UMovieSceneBindingExtensions::GetTracks(const FMovieSceneBindingProxy& InBinding)
 {
 	UMovieScene* MovieScene = InBinding.GetMovieScene();
 	if (MovieScene)
@@ -102,7 +102,7 @@
 	return TArray<UMovieSceneTrack*>();
 }
 
-void UMovieSceneBindingExtensions::RemoveTrack(const FSequencerBindingProxy& InBinding, UMovieSceneTrack* TrackToRemove)
+void UMovieSceneBindingExtensions::RemoveTrack(const FMovieSceneBindingProxy& InBinding, UMovieSceneTrack* TrackToRemove)
 {
 	if (!TrackToRemove)
 	{
@@ -117,7 +117,7 @@
 	}
 }
 
-void UMovieSceneBindingExtensions::Remove(const FSequencerBindingProxy& InBinding)
+void UMovieSceneBindingExtensions::Remove(const FMovieSceneBindingProxy& InBinding)
 {
 	UMovieScene* MovieScene = InBinding.GetMovieScene();
 	if (MovieScene)
@@ -129,7 +129,7 @@
 	}
 }
 
-TArray<UMovieSceneTrack*> UMovieSceneBindingExtensions::FindTracksByType(const FSequencerBindingProxy& InBinding, TSubclassOf<UMovieSceneTrack> TrackType)
+TArray<UMovieSceneTrack*> UMovieSceneBindingExtensions::FindTracksByType(const FMovieSceneBindingProxy& InBinding, TSubclassOf<UMovieSceneTrack> TrackType)
 {
 	UMovieScene* MovieScene   = InBinding.GetMovieScene();
 	UClass*      DesiredClass = TrackType.Get();
@@ -146,7 +146,7 @@
 	return TArray<UMovieSceneTrack*>();
 }
 
-TArray<UMovieSceneTrack*> UMovieSceneBindingExtensions::FindTracksByExactType(const FSequencerBindingProxy& InBinding, TSubclassOf<UMovieSceneTrack> TrackType)
+TArray<UMovieSceneTrack*> UMovieSceneBindingExtensions::FindTracksByExactType(const FMovieSceneBindingProxy& InBinding, TSubclassOf<UMovieSceneTrack> TrackType)
 {
 	UMovieScene* MovieScene   = InBinding.GetMovieScene();
 	UClass*      DesiredClass = TrackType.Get();
@@ -163,7 +163,7 @@
 	return TArray<UMovieSceneTrack*>();
 }
 
-UMovieSceneTrack* UMovieSceneBindingExtensions::AddTrack(const FSequencerBindingProxy& InBinding, TSubclassOf<UMovieSceneTrack> TrackType)
+UMovieSceneTrack* UMovieSceneBindingExtensions::AddTrack(const FMovieSceneBindingProxy& InBinding, TSubclassOf<UMovieSceneTrack> TrackType)
 {
 	UMovieScene* MovieScene   = InBinding.GetMovieScene();
 	UClass*      DesiredClass = TrackType.Get();
@@ -184,9 +184,9 @@
 	return nullptr;
 }
 
-TArray<FSequencerBindingProxy> UMovieSceneBindingExtensions::GetChildPossessables(const FSequencerBindingProxy& InBinding)
+TArray<FMovieSceneBindingProxy> UMovieSceneBindingExtensions::GetChildPossessables(const FMovieSceneBindingProxy& InBinding)
 {
-	TArray<FSequencerBindingProxy> Result;
+	TArray<FMovieSceneBindingProxy> Result;
 
 	UMovieScene* MovieScene = InBinding.GetMovieScene();
 	if (MovieScene)
@@ -214,7 +214,7 @@
 	return Result;
 }
 
-UObject* UMovieSceneBindingExtensions::GetObjectTemplate(const FSequencerBindingProxy& InBinding)
+UObject* UMovieSceneBindingExtensions::GetObjectTemplate(const FMovieSceneBindingProxy& InBinding)
 {
 	UMovieScene* MovieScene = InBinding.GetMovieScene();
 	if (MovieScene)
@@ -228,7 +228,7 @@
 	return nullptr;
 }
 
-UClass* UMovieSceneBindingExtensions::GetPossessedObjectClass(const FSequencerBindingProxy& InBinding)
+UClass* UMovieSceneBindingExtensions::GetPossessedObjectClass(const FMovieSceneBindingProxy& InBinding)
 {
 	UMovieScene* MovieScene = InBinding.GetMovieScene();
 	if (MovieScene)
@@ -244,7 +244,7 @@
 	return nullptr;
 }
 
-FSequencerBindingProxy UMovieSceneBindingExtensions::GetParent(const FSequencerBindingProxy& InBinding)
+FMovieSceneBindingProxy UMovieSceneBindingExtensions::GetParent(const FMovieSceneBindingProxy& InBinding)
 {
 	UMovieScene* MovieScene = InBinding.GetMovieScene();
 	if (MovieScene)
@@ -252,14 +252,14 @@
 		FMovieScenePossessable* Possessable = MovieScene->FindPossessable(InBinding.BindingID);
 		if (Possessable)
 		{
-			return FSequencerBindingProxy(Possessable->GetParent(), InBinding.Sequence);
+			return FMovieSceneBindingProxy(Possessable->GetParent(), InBinding.Sequence);
 		}
 	}
-	return FSequencerBindingProxy();
+	return FMovieSceneBindingProxy();
 }
 
 
-void UMovieSceneBindingExtensions::SetParent(const FSequencerBindingProxy& InBinding, const FSequencerBindingProxy& InParentBinding)
+void UMovieSceneBindingExtensions::SetParent(const FMovieSceneBindingProxy& InBinding, const FMovieSceneBindingProxy& InParentBinding)
 {
 	UMovieScene* MovieScene = InBinding.GetMovieScene();
 	if (MovieScene)
@@ -273,7 +273,7 @@
 	}
 }
 
-void UMovieSceneBindingExtensions::MoveBindingContents(const FSequencerBindingProxy& SourceBindingId, const FSequencerBindingProxy& DestinationBindingId)
+void UMovieSceneBindingExtensions::MoveBindingContents(const FMovieSceneBindingProxy& SourceBindingId, const FMovieSceneBindingProxy& DestinationBindingId)
 {
 	UMovieScene* MovieScene = SourceBindingId.GetMovieScene();
 	if (MovieScene)
diff --git a/Engine/Plugins/MovieScene/SequencerScripting/Source/SequencerScripting/Private/ExtensionLibraries/MovieSceneFolderExtensions.cpp b/Engine/Plugins/MovieScene/SequencerScripting/Source/SequencerScripting/Private/ExtensionLibraries/MovieSceneFolderExtensions.cpp
--- a/Engine/Plugins/MovieScene/SequencerScripting/Source/SequencerScripting/Private/ExtensionLibraries/MovieSceneFolderExtensions.cpp
+++ b/Engine/Plugins/MovieScene/SequencerScripting/Source/SequencerScripting/Private/ExtensionLibraries/MovieSceneFolderExtensions.cpp
@@ -152,15 +152,15 @@
 	return true;
 }
 
-TArray<FSequencerBindingProxy> UMovieSceneFolderExtensions::GetChildObjectBindings(UMovieSceneFolder* Folder)
+TArray<FMovieSceneBindingProxy> UMovieSceneFolderExtensions::GetChildObjectBindings(UMovieSceneFolder* Folder)
 {
 	if (!Folder)
 	{
 		FFrame::KismetExecutionMessage(TEXT("Cannot call GetChildObjectBindings with a null folder"), ELogVerbosity::Error);
-		return TArray<FSequencerBindingProxy>();
+		return TArray<FMovieSceneBindingProxy>();
 	}
 
-	TArray<FSequencerBindingProxy> Result;
+	TArray<FMovieSceneBindingProxy> Result;
 
 	// Attempt to get the sequence reference from the folder
 	UMovieScene* MovieScene = Cast<UMovieScene>(Folder->GetOuter());
@@ -168,13 +168,13 @@
 
 	for (FGuid ID : Folder->GetChildObjectBindings())
 	{
-		Result.Add(FSequencerBindingProxy(ID, Sequence));
+		Result.Add(FMovieSceneBindingProxy(ID, Sequence));
 	}
 
 	return Result;
 }
 
-bool UMovieSceneFolderExtensions::AddChildObjectBinding(UMovieSceneFolder* Folder, FSequencerBindingProxy InObjectBinding)
+bool UMovieSceneFolderExtensions::AddChildObjectBinding(UMovieSceneFolder* Folder, FMovieSceneBindingProxy InObjectBinding)
 {
 	if (!Folder)
 	{
@@ -191,7 +191,7 @@
 	return false;
 }
 
-bool UMovieSceneFolderExtensions::RemoveChildObjectBinding(UMovieSceneFolder* Folder, const FSequencerBindingProxy InObjectBinding)
+bool UMovieSceneFolderExtensions::RemoveChildObjectBinding(UMovieSceneFolder* Folder, const FMovieSceneBindingProxy InObjectBinding)
 {
 	if (!Folder)
 	{
diff --git a/Engine/Plugins/MovieScene/SequencerScripting/Source/SequencerScripting/Private/ExtensionLibraries/MovieSceneSequenceExtensions.cpp b/Engine/Plugins/MovieScene/SequencerScripting/Source/SequencerScripting/Private/ExtensionLibraries/MovieSceneSequenceExtensions.cpp
--- a/Engine/Plugins/MovieScene/SequencerScripting/Source/SequencerScripting/Private/ExtensionLibraries/MovieSceneSequenceExtensions.cpp
+++ b/Engine/Plugins/MovieScene/SequencerScripting/Source/SequencerScripting/Private/ExtensionLibraries/MovieSceneSequenceExtensions.cpp
@@ -660,12 +660,12 @@
 	return MovieScene->GetEarliestTimecodeSource().Timecode;
 }
 
-FSequencerBindingProxy UMovieSceneSequenceExtensions::FindBindingByName(UMovieSceneSequence* Sequence, FString Name)
+FMovieSceneBindingProxy UMovieSceneSequenceExtensions::FindBindingByName(UMovieSceneSequence* Sequence, FString Name)
 {
 	if (!Sequence)
 	{
 		FFrame::KismetExecutionMessage(TEXT("Cannot call FindBindingByName on a null sequence"), ELogVerbosity::Error);
-		return FSequencerBindingProxy();
+		return FMovieSceneBindingProxy();
 	}
 
 	UMovieScene* MovieScene = GetMovieScene(Sequence);
@@ -674,18 +674,18 @@
 		const FMovieSceneBinding* Binding = Algo::FindBy(MovieScene->GetBindings(), Name, &FMovieSceneBinding::GetName);
 		if (Binding)
 		{
-			return FSequencerBindingProxy(Binding->GetObjectGuid(), Sequence);
+			return FMovieSceneBindingProxy(Binding->GetObjectGuid(), Sequence);
 		}
 	}
-	return FSequencerBindingProxy();
+	return FMovieSceneBindingProxy();
 }
 
-FSequencerBindingProxy UMovieSceneSequenceExtensions::FindBindingById(UMovieSceneSequence* Sequence, FGuid BindingId)
+FMovieSceneBindingProxy UMovieSceneSequenceExtensions::FindBindingById(UMovieSceneSequence* Sequence, FGuid BindingId)
 {
 	if (!Sequence)
 	{
 		FFrame::KismetExecutionMessage(TEXT("Cannot call FindBindingById on a null sequence"), ELogVerbosity::Error);
-		return FSequencerBindingProxy();
+		return FMovieSceneBindingProxy();
 	}
 
 	UMovieScene* MovieScene = GetMovieScene(Sequence);
@@ -694,21 +694,21 @@
 		const FMovieSceneBinding* Binding = Algo::FindBy(MovieScene->GetBindings(), BindingId, &FMovieSceneBinding::GetObjectGuid);
 		if (Binding)
 		{
-			return FSequencerBindingProxy(Binding->GetObjectGuid(), Sequence);
+			return FMovieSceneBindingProxy(Binding->GetObjectGuid(), Sequence);
 		}
 	}
-	return FSequencerBindingProxy();
+	return FMovieSceneBindingProxy();
 }
 
-TArray<FSequencerBindingProxy> UMovieSceneSequenceExtensions::GetBindings(UMovieSceneSequence* Sequence)
+TArray<FMovieSceneBindingProxy> UMovieSceneSequenceExtensions::GetBindings(UMovieSceneSequence* Sequence)
 {
 	if (!Sequence)
 	{
 		FFrame::KismetExecutionMessage(TEXT("Cannot call GetBindings on a null sequence"), ELogVerbosity::Error);
-		return TArray<FSequencerBindingProxy>();
+		return TArray<FMovieSceneBindingProxy>();
 	}
 
-	TArray<FSequencerBindingProxy> AllBindings;
+	TArray<FMovieSceneBindingProxy> AllBindings;
 
 	UMovieScene* MovieScene = GetMovieScene(Sequence);
 	if (MovieScene)
@@ -722,15 +722,15 @@
 	return AllBindings;
 }
 
-TArray<FSequencerBindingProxy> UMovieSceneSequenceExtensions::GetSpawnables(UMovieSceneSequence* Sequence)
+TArray<FMovieSceneBindingProxy> UMovieSceneSequenceExtensions::GetSpawnables(UMovieSceneSequence* Sequence)
 {
 	if (!Sequence)
 	{
 		FFrame::KismetExecutionMessage(TEXT("Cannot call GetSpawnables on a null sequence"), ELogVerbosity::Error);
-		return TArray<FSequencerBindingProxy>();
+		return TArray<FMovieSceneBindingProxy>();
 	}
 
-	TArray<FSequencerBindingProxy> AllSpawnables;
+	TArray<FMovieSceneBindingProxy> AllSpawnables;
 
 	UMovieScene* MovieScene = GetMovieScene(Sequence);
 	if (MovieScene)
@@ -746,15 +746,15 @@
 	return AllSpawnables;
 }
 
-TArray<FSequencerBindingProxy> UMovieSceneSequenceExtensions::GetPossessables(UMovieSceneSequence* Sequence)
+TArray<FMovieSceneBindingProxy> UMovieSceneSequenceExtensions::GetPossessables(UMovieSceneSequence* Sequence)
 {
 	if (!Sequence)
 	{
 		FFrame::KismetExecutionMessage(TEXT("Cannot call GetPossessables on a null sequence"), ELogVerbosity::Error);
-		return TArray<FSequencerBindingProxy>();
+		return TArray<FMovieSceneBindingProxy>();
 	}
 
-	TArray<FSequencerBindingProxy> AllPossessables;
+	TArray<FMovieSceneBindingProxy> AllPossessables;
 
 	UMovieScene* MovieScene = GetMovieScene(Sequence);
 	if (MovieScene)
@@ -770,61 +770,61 @@
 	return AllPossessables;
 }
 
-FSequencerBindingProxy UMovieSceneSequenceExtensions::AddPossessable(UMovieSceneSequence* Sequence, UObject* ObjectToPossess)
+FMovieSceneBindingProxy UMovieSceneSequenceExtensions::AddPossessable(UMovieSceneSequence* Sequence, UObject* ObjectToPossess)
 {
 	if (!Sequence)
 	{
 		FFrame::KismetExecutionMessage(TEXT("Cannot call AddPossessable on a null sequence"), ELogVerbosity::Error);
-		return FSequencerBindingProxy();
+		return FMovieSceneBindingProxy();
 	}
 
 	if (!ObjectToPossess)
 	{
 		FFrame::KismetExecutionMessage(TEXT("Cannot call AddPossessable on a null object"), ELogVerbosity::Error);
-		return FSequencerBindingProxy();
+		return FMovieSceneBindingProxy();
 	}
 
 	FGuid NewGuid = Sequence->CreatePossessable(ObjectToPossess);
-	return FSequencerBindingProxy(NewGuid, Sequence);
+	return FMovieSceneBindingProxy(NewGuid, Sequence);
 }
 
-FSequencerBindingProxy UMovieSceneSequenceExtensions::AddSpawnableFromInstance(UMovieSceneSequence* Sequence, UObject* ObjectToSpawn)
+FMovieSceneBindingProxy UMovieSceneSequenceExtensions::AddSpawnableFromInstance(UMovieSceneSequence* Sequence, UObject* ObjectToSpawn)
 {
 	if (!Sequence)
 	{
 		FFrame::KismetExecutionMessage(TEXT("Cannot call AddSpawnableFromInstance on a null sequence"), ELogVerbosity::Error);
-		return FSequencerBindingProxy();
+		return FMovieSceneBindingProxy();
 	}
 
 	if (!ObjectToSpawn)
 	{
 		FFrame::KismetExecutionMessage(TEXT("Cannot call AddSpawnableFromInstance on a null object"), ELogVerbosity::Error);
-		return FSequencerBindingProxy();
+		return FMovieSceneBindingProxy();
 	}
 
 	FGuid NewGuid = Sequence->AllowsSpawnableObjects() ? Sequence->CreateSpawnable(ObjectToSpawn) : FGuid();
-	return FSequencerBindingProxy(NewGuid, Sequence);
+	return FMovieSceneBindingProxy(NewGuid, Sequence);
 }
 
-FSequencerBindingProxy UMovieSceneSequenceExtensions::AddSpawnableFromClass(UMovieSceneSequence* Sequence, UClass* ClassToSpawn)
+FMovieSceneBindingProxy UMovieSceneSequenceExtensions::AddSpawnableFromClass(UMovieSceneSequence* Sequence, UClass* ClassToSpawn)
 {
 	if (!Sequence)
 	{
 		FFrame::KismetExecutionMessage(TEXT("Cannot call AddSpawnableFromClass on a null sequence"), ELogVerbosity::Error);
-		return FSequencerBindingProxy();
+		return FMovieSceneBindingProxy();
 	}
 
 	if (!ClassToSpawn)
 	{
 		FFrame::KismetExecutionMessage(TEXT("Cannot call AddSpawnableFromClass on a null class"), ELogVerbosity::Error);
-		return FSequencerBindingProxy();
+		return FMovieSceneBindingProxy();
 	}
 
 	FGuid NewGuid = Sequence->AllowsSpawnableObjects() ? Sequence->CreateSpawnable(ClassToSpawn) : FGuid();
-	return FSequencerBindingProxy(NewGuid, Sequence);
+	return FMovieSceneBindingProxy(NewGuid, Sequence);
 }
 
-TArray<UObject*> UMovieSceneSequenceExtensions::LocateBoundObjects(UMovieSceneSequence* Sequence, const FSequencerBindingProxy& InBinding, UObject* Context)
+TArray<UObject*> UMovieSceneSequenceExtensions::LocateBoundObjects(UMovieSceneSequence* Sequence, const FMovieSceneBindingProxy& InBinding, UObject* Context)
 {
 	if (!Sequence)
 	{
@@ -840,7 +840,7 @@
 	return Result;
 }
 
-FMovieSceneObjectBindingID UMovieSceneSequenceExtensions::MakeBindingID(UMovieSceneSequence* MasterSequence, const FSequencerBindingProxy& InBinding, EMovieSceneObjectBindingSpace Space)
+FMovieSceneObjectBindingID UMovieSceneSequenceExtensions::MakeBindingID(UMovieSceneSequence* MasterSequence, const FMovieSceneBindingProxy& InBinding, EMovieSceneObjectBindingSpace Space)
 {
 	if (!MasterSequence)
 	{
@@ -858,12 +858,12 @@
 	return BindingID;
 }
 
-FMovieSceneObjectBindingID UMovieSceneSequenceExtensions::GetBindingID(const FSequencerBindingProxy& InBinding)
+FMovieSceneObjectBindingID UMovieSceneSequenceExtensions::GetBindingID(const FMovieSceneBindingProxy& InBinding)
 {
 	return UE::MovieScene::FRelativeObjectBindingID(InBinding.BindingID);
 }
 
-FMovieSceneObjectBindingID UMovieSceneSequenceExtensions::GetPortableBindingID(UMovieSceneSequence* MasterSequence, UMovieSceneSequence* DestinationSequence, const FSequencerBindingProxy& InBinding)
+FMovieSceneObjectBindingID UMovieSceneSequenceExtensions::GetPortableBindingID(UMovieSceneSequence* MasterSequence, UMovieSceneSequence* DestinationSequence, const FMovieSceneBindingProxy& InBinding)
 {
 	if (!MasterSequence || !DestinationSequence || !InBinding.Sequence)
 	{
@@ -942,12 +942,12 @@
 	return UE::MovieScene::FRelativeObjectBindingID(DestinationSequenceID.GetValue(), TargetSequenceID.GetValue(), InBinding.BindingID, Hierarchy);
 }
 
-FSequencerBindingProxy UMovieSceneSequenceExtensions::ResolveBindingID(UMovieSceneSequence* MasterSequence, FMovieSceneObjectBindingID InObjectBindingID)
+FMovieSceneBindingProxy UMovieSceneSequenceExtensions::ResolveBindingID(UMovieSceneSequence* MasterSequence, FMovieSceneObjectBindingID InObjectBindingID)
 {
 	if (!MasterSequence)
 	{
 		FFrame::KismetExecutionMessage(TEXT("Cannot call ResolveBindingID on a null sequence"), ELogVerbosity::Error);
-		return FSequencerBindingProxy();
+		return FMovieSceneBindingProxy();
 	}
 
 	UMovieSceneSequence* Sequence = MasterSequence;
@@ -963,7 +963,7 @@
 		}
 	}
 
-	return FSequencerBindingProxy(InObjectBindingID.GetGuid(), Sequence);
+	return FMovieSceneBindingProxy(InObjectBindingID.GetGuid(), Sequence);
 }
 
 TArray<UMovieSceneFolder*> UMovieSceneSequenceExtensions::GetRootFoldersInSequence(UMovieSceneSequence* Sequence)
diff --git a/Engine/Plugins/MovieScene/SequencerScripting/Source/SequencerScripting/Private/KeysAndChannels/MovieSceneScriptingActorReference.h b/Engine/Plugins/MovieScene/SequencerScripting/Source/SequencerScripting/Private/KeysAndChannels/MovieSceneScriptingActorReference.h
--- a/Engine/Plugins/MovieScene/SequencerScripting/Source/SequencerScripting/Private/KeysAndChannels/MovieSceneScriptingActorReference.h
+++ b/Engine/Plugins/MovieScene/SequencerScripting/Source/SequencerScripting/Private/KeysAndChannels/MovieSceneScriptingActorReference.h
@@ -11,7 +11,6 @@
 #include "KeysAndChannels/MovieSceneScriptingChannel.h"
 #include "KeyParams.h"
 #include "MovieScene.h"
-#include "SequencerBindingProxy.h"
 
 #include "MovieSceneScriptingActorReference.generated.h"
 
diff --git a/Engine/Plugins/MovieScene/SequencerScripting/Source/SequencerScripting/Private/KeysAndChannels/MovieSceneScriptingObjectPath.h b/Engine/Plugins/MovieScene/SequencerScripting/Source/SequencerScripting/Private/KeysAndChannels/MovieSceneScriptingObjectPath.h
--- a/Engine/Plugins/MovieScene/SequencerScripting/Source/SequencerScripting/Private/KeysAndChannels/MovieSceneScriptingObjectPath.h
+++ b/Engine/Plugins/MovieScene/SequencerScripting/Source/SequencerScripting/Private/KeysAndChannels/MovieSceneScriptingObjectPath.h
@@ -13,7 +13,6 @@
 #include "MovieScene.h"
 #include "MovieSceneSection.h"
 #include "MovieSceneSequence.h"
-#include "SequencerBindingProxy.h"
 
 #include "MovieSceneScriptingObjectPath.generated.h"
 
diff --git a/Engine/Plugins/MovieScene/SequencerScripting/Source/SequencerScripting/Private/SequencerBindingProxy.cpp b/Engine/Plugins/MovieScene/SequencerScripting/Source/SequencerScripting/Private/SequencerBindingProxy.cpp
deleted file mode 100644
diff --git a/Engine/Plugins/MovieScene/SequencerScripting/Source/SequencerScripting/Public/ExtensionLibraries/MovieSceneBindingExtensions.h b/Engine/Plugins/MovieScene/SequencerScripting/Source/SequencerScripting/Public/ExtensionLibraries/MovieSceneBindingExtensions.h
--- a/Engine/Plugins/MovieScene/SequencerScripting/Source/SequencerScripting/Public/ExtensionLibraries/MovieSceneBindingExtensions.h
+++ b/Engine/Plugins/MovieScene/SequencerScripting/Source/SequencerScripting/Public/ExtensionLibraries/MovieSceneBindingExtensions.h
@@ -7,7 +7,7 @@
 #include "Kismet/BlueprintFunctionLibrary.h"
 #include "MovieSceneTrack.h"
 #include "Templates/SubclassOf.h"
-#include "SequencerBindingProxy.h"
+#include "MovieSceneBindingProxy.h"
 
 #include "MovieSceneBindingExtensions.generated.h"
 
@@ -24,7 +24,7 @@ public:
 	 * Check whether the specified binding is valid
 	 */
 	UFUNCTION(BlueprintPure, Category = "Sequencer|Sequence", meta=(ScriptMethod))
-	static bool IsValid(const FSequencerBindingProxy& InBinding);
+	static bool IsValid(const FMovieSceneBindingProxy& InBinding);
 
 	/**
 	 * Get this binding's ID
@@ -33,7 +33,7 @@ public:
 	 * @return The guid that uniquely represents this binding
 	 */
 	UFUNCTION(BlueprintCallable, Category = "Sequencer|Sequence", meta=(ScriptMethod))
-	static FGuid GetId(const FSequencerBindingProxy& InBinding);
+	static FGuid GetId(const FMovieSceneBindingProxy& InBinding);
 
 	/**
 	 * Get this binding's name
@@ -42,7 +42,7 @@ public:
 	 * @return The display name of the binding
 	 */
 	UFUNCTION(BlueprintCallable, Category = "Sequencer|Sequence", meta=(ScriptMethod))
-	static FText GetDisplayName(const FSequencerBindingProxy& InBinding);
+	static FText GetDisplayName(const FMovieSceneBindingProxy& InBinding);
 
 	/**
 	 * Set this binding's name
@@ -51,7 +51,7 @@ public:
 	 * @param InName The display name of the binding
 	 */
 	UFUNCTION(BlueprintCallable, Category = "Sequencer|Sequence", meta=(ScriptMethod, DevelopmentOnly))
-	static void SetDisplayName(const FSequencerBindingProxy& InBinding, const FText& InDisplayName);
+	static void SetDisplayName(const FMovieSceneBindingProxy& InBinding, const FText& InDisplayName);
 
 	/**
 	 * Get this binding's object non-display name
@@ -60,7 +60,7 @@ public:
 	 * @return The name of the binding
 	 */
 	UFUNCTION(BlueprintCallable, Category = "Sequencer|Sequence", meta=(ScriptMethod))
-	static FString GetName(const FSequencerBindingProxy& InBinding);
+	static FString GetName(const FMovieSceneBindingProxy& InBinding);
 
 	/**
 	 * Set this binding's object non-display name
@@ -69,7 +69,7 @@ public:
 	 * @param InName The name of the binding
 	 */
 	UFUNCTION(BlueprintCallable, Category = "Sequencer|Sequence", meta=(ScriptMethod))
-	static void SetName(const FSequencerBindingProxy& InBinding, const FString& InName);
+	static void SetName(const FMovieSceneBindingProxy& InBinding, const FString& InName);
 
 	/**
 	 * Get all the tracks stored within this binding
@@ -78,7 +78,7 @@ public:
 	 * @return An array containing all the binding's tracks
 	 */
 	UFUNCTION(BlueprintCallable, Category = "Sequencer|Sequence", meta=(ScriptMethod))
-	static TArray<UMovieSceneTrack*> GetTracks(const FSequencerBindingProxy& InBinding);
+	static TArray<UMovieSceneTrack*> GetTracks(const FMovieSceneBindingProxy& InBinding);
 
 	/**
 	 * Find all tracks within a given binding of the specified type
@@ -88,7 +88,7 @@ public:
 	 * @return An array containing any tracks that match the type specified
 	 */
 	UFUNCTION(BlueprintCallable, Category = "Sequencer|Sequence", meta=(ScriptMethod, DeterminesOutputType="TrackType"))
-	static TArray<UMovieSceneTrack*> FindTracksByType(const FSequencerBindingProxy& InBinding, TSubclassOf<UMovieSceneTrack> TrackType);
+	static TArray<UMovieSceneTrack*> FindTracksByType(const FMovieSceneBindingProxy& InBinding, TSubclassOf<UMovieSceneTrack> TrackType);
 
 	/**
 	 * Find all tracks within a given binding of the specified type, not allowing sub-classed types
@@ -98,7 +98,7 @@ public:
 	 * @return An array containing any tracks that are exactly the same as the type specified
 	 */
 	UFUNCTION(BlueprintCallable, Category = "Sequencer|Sequence", meta=(ScriptMethod, DeterminesOutputType="TrackType"))
-	static TArray<UMovieSceneTrack*> FindTracksByExactType(const FSequencerBindingProxy& InBinding, TSubclassOf<UMovieSceneTrack> TrackType);
+	static TArray<UMovieSceneTrack*> FindTracksByExactType(const FMovieSceneBindingProxy& InBinding, TSubclassOf<UMovieSceneTrack> TrackType);
 
 	/**
 	 * Remove the specified track from this binding
@@ -107,7 +107,7 @@ public:
 	 * @param TrackToRemove The track to remove
 	 */
 	UFUNCTION(BlueprintCallable, Category = "Sequencer|Sequence", meta=(ScriptMethod))
-	static void RemoveTrack(const FSequencerBindingProxy& InBinding, UMovieSceneTrack* TrackToRemove);
+	static void RemoveTrack(const FMovieSceneBindingProxy& InBinding, UMovieSceneTrack* TrackToRemove);
 
 	/**
 	 * Remove the specified binding
@@ -115,7 +115,7 @@ public:
 	 * @param InBinding     The binding to remove the track from
 	 */
 	UFUNCTION(BlueprintCallable, Category = "Sequencer|Sequence", meta=(ScriptMethod))
-	static void Remove(const FSequencerBindingProxy& InBinding);
+	static void Remove(const FMovieSceneBindingProxy& InBinding);
 
 	/**
 	 * Add a new track to the specified binding
@@ -125,7 +125,7 @@ public:
 	 * @return The newly created track, if successful
 	 */
 	UFUNCTION(BlueprintCallable, Category = "Sequencer|Sequence", meta=(ScriptMethod))
-	static UMovieSceneTrack* AddTrack(const FSequencerBindingProxy& InBinding, TSubclassOf<UMovieSceneTrack> TrackType);
+	static UMovieSceneTrack* AddTrack(const FMovieSceneBindingProxy& InBinding, TSubclassOf<UMovieSceneTrack> TrackType);
 
 	/**
 	* Get all the children of this binding
@@ -134,7 +134,7 @@ public:
 	* @return An array containing all the binding's children
 	*/
 	UFUNCTION(BlueprintCallable, Category = "Sequencer|Sequence", meta=(ScriptMethod))
-	static TArray<FSequencerBindingProxy> GetChildPossessables(const FSequencerBindingProxy& InBinding);
+	static TArray<FMovieSceneBindingProxy> GetChildPossessables(const FMovieSceneBindingProxy& InBinding);
 
 	/**
 	* Get this binding's object template
@@ -143,7 +143,7 @@ public:
 	* @return The object template of the binding
 	*/
 	UFUNCTION(BlueprintCallable, Category = "Sequencer|Sequence", meta=(ScriptMethod))
-	static UObject* GetObjectTemplate(const FSequencerBindingProxy& InBinding);
+	static UObject* GetObjectTemplate(const FMovieSceneBindingProxy& InBinding);
 
 	/**
 	* Get this binding's possessed object class
@@ -152,7 +152,7 @@ public:
 	* @return The possessed object class of the binding
 	*/
 	UFUNCTION(BlueprintCallable, Category = "Sequencer|Sequence", meta=(ScriptMethod))
-	static UClass* GetPossessedObjectClass(const FSequencerBindingProxy& InBinding);
+	static UClass* GetPossessedObjectClass(const FMovieSceneBindingProxy& InBinding);
 
 	/**
 	* Get the parent of this binding
@@ -161,7 +161,7 @@ public:
 	* @return The binding's parent
 	*/
 	UFUNCTION(BlueprintCallable, Category = "Sequencer|Sequence", meta=(ScriptMethod))
-	static FSequencerBindingProxy GetParent(const FSequencerBindingProxy& InBinding);
+	static FMovieSceneBindingProxy GetParent(const FMovieSceneBindingProxy& InBinding);
 
 	/**
 	 * Set the parent to this binding
@@ -170,7 +170,7 @@ public:
 	 * @param InParentBinding     The parent to set the InBinding to
 	 */
 	UFUNCTION(BlueprintCallable, Category = "Sequencer|Sequence", meta = (ScriptMethod))
-	static void SetParent(const FSequencerBindingProxy& InBinding, const FSequencerBindingProxy& InParentBinding);
+	static void SetParent(const FMovieSceneBindingProxy& InBinding, const FMovieSceneBindingProxy& InParentBinding);
 
 	/**
      * Move all the contents (tracks, child bindings) of the specified binding ID onto another
@@ -179,5 +179,5 @@ public:
 	 * @param DestinationBindingId The identifier of the binding ID to move the contents to	 
 	 */
 	UFUNCTION(BlueprintCallable, Category = "Sequencer|Sequence", meta = (ScriptMethod))
-	static void MoveBindingContents(const FSequencerBindingProxy& SourceBindingId, const FSequencerBindingProxy& DestinationBindingId);
+	static void MoveBindingContents(const FMovieSceneBindingProxy& SourceBindingId, const FMovieSceneBindingProxy& DestinationBindingId);
 };
\ No newline at end of file
diff --git a/Engine/Plugins/MovieScene/SequencerScripting/Source/SequencerScripting/Public/ExtensionLibraries/MovieSceneFolderExtensions.h b/Engine/Plugins/MovieScene/SequencerScripting/Source/SequencerScripting/Public/ExtensionLibraries/MovieSceneFolderExtensions.h
--- a/Engine/Plugins/MovieScene/SequencerScripting/Source/SequencerScripting/Public/ExtensionLibraries/MovieSceneFolderExtensions.h
+++ b/Engine/Plugins/MovieScene/SequencerScripting/Source/SequencerScripting/Public/ExtensionLibraries/MovieSceneFolderExtensions.h
@@ -5,7 +5,7 @@
 #include "CoreTypes.h"
 #include "UObject/ObjectMacros.h"
 #include "Kismet/BlueprintFunctionLibrary.h"
-#include "SequencerBindingProxy.h"
+#include "MovieSceneBindingProxy.h"
 
 #include "MovieSceneFolderExtensions.generated.h"
 
@@ -124,7 +124,7 @@ public:
 	 * @return The object bindings under the given folder
 	 */
 	UFUNCTION(BlueprintCallable, BlueprintPure, Category = "Sequencer|Sequence", meta=(ScriptMethod))
-	static TArray<FSequencerBindingProxy> GetChildObjectBindings(UMovieSceneFolder* Folder);
+	static TArray<FMovieSceneBindingProxy> GetChildObjectBindings(UMovieSceneFolder* Folder);
 
 	/** 
 	 * Add a guid for an object binding to this folder 
@@ -134,7 +134,7 @@ public:
 	 * @return True if the addition is successful
 	 */
 	UFUNCTION(BlueprintCallable, Category = "Sequencer|Sequence", meta=(ScriptMethod))
-	static bool AddChildObjectBinding(UMovieSceneFolder* Folder, FSequencerBindingProxy InObjectBinding);
+	static bool AddChildObjectBinding(UMovieSceneFolder* Folder, FMovieSceneBindingProxy InObjectBinding);
 
 	/** 
 	 * Remove an object binding from the given folder
@@ -144,5 +144,5 @@ public:
 	 * @return True if the operation succeeds
 	 */
 	UFUNCTION(BlueprintCallable, Category = "Sequencer|Sequence", meta=(ScriptMethod))
-	static bool RemoveChildObjectBinding(UMovieSceneFolder* Folder, const FSequencerBindingProxy InObjectBinding);
+	static bool RemoveChildObjectBinding(UMovieSceneFolder* Folder, const FMovieSceneBindingProxy InObjectBinding);
 };
\ No newline at end of file
diff --git a/Engine/Plugins/MovieScene/SequencerScripting/Source/SequencerScripting/Public/ExtensionLibraries/MovieSceneSequenceExtensions.h b/Engine/Plugins/MovieScene/SequencerScripting/Source/SequencerScripting/Public/ExtensionLibraries/MovieSceneSequenceExtensions.h
--- a/Engine/Plugins/MovieScene/SequencerScripting/Source/SequencerScripting/Public/ExtensionLibraries/MovieSceneSequenceExtensions.h
+++ b/Engine/Plugins/MovieScene/SequencerScripting/Source/SequencerScripting/Public/ExtensionLibraries/MovieSceneSequenceExtensions.h
@@ -6,7 +6,7 @@
 #include "UObject/ObjectMacros.h"
 #include "Containers/ArrayView.h"
 #include "Kismet/BlueprintFunctionLibrary.h"
-#include "SequencerBindingProxy.h"
+#include "MovieSceneBindingProxy.h"
 #include "SequencerScriptingRange.h"
 #include "Templates/SubclassOf.h"
 #include "MovieSceneTrack.h"
@@ -355,7 +355,7 @@ public:
 	 * @return A unique identifier for the binding, or invalid
 	 */
 	UFUNCTION(BlueprintCallable, Category="Sequencer|Sequence", meta=(ScriptMethod))
-	static FSequencerBindingProxy FindBindingByName(UMovieSceneSequence* Sequence, FString Name);
+	static FMovieSceneBindingProxy FindBindingByName(UMovieSceneSequence* Sequence, FString Name);
 
 	/**
 	 * Attempt to locate a binding in this sequence by its Id
@@ -365,7 +365,7 @@ public:
 	 * @return A unique identifier for the binding, or invalid
 	 */
 	UFUNCTION(BlueprintCallable, Category="Sequencer|Sequence", meta=(ScriptMethod))
-	static FSequencerBindingProxy FindBindingById(UMovieSceneSequence* Sequence, FGuid BindingId);
+	static FMovieSceneBindingProxy FindBindingById(UMovieSceneSequence* Sequence, FGuid BindingId);
 
 	/**
 	 * Get all the bindings in this sequence
@@ -374,7 +374,7 @@ public:
 	 * @return An array of unique identifiers for all the bindings in this sequence
 	 */
 	UFUNCTION(BlueprintCallable, Category="Sequencer|Sequence", meta=(ScriptMethod))
-	static TArray<FSequencerBindingProxy> GetBindings(UMovieSceneSequence* Sequence);
+	static TArray<FMovieSceneBindingProxy> GetBindings(UMovieSceneSequence* Sequence);
 
 	/**
 	* Get all the spawnables in this sequence
@@ -383,7 +383,7 @@ public:
 	* @return Spawnables in this sequence
 	*/
 	UFUNCTION(BlueprintCallable, Category="Sequencer|Sequence", meta=(ScriptMethod))
-	static TArray<FSequencerBindingProxy> GetSpawnables(UMovieSceneSequence* Sequence);
+	static TArray<FMovieSceneBindingProxy> GetSpawnables(UMovieSceneSequence* Sequence);
 
 	/**
 	* Get all the possessables in this sequence
@@ -392,7 +392,7 @@ public:
 	* @return Possessables in this sequence
 	*/
 	UFUNCTION(BlueprintCallable, Category="Sequencer|Sequence", meta=(ScriptMethod))
-	static TArray<FSequencerBindingProxy> GetPossessables(UMovieSceneSequence* Sequence);
+	static TArray<FMovieSceneBindingProxy> GetPossessables(UMovieSceneSequence* Sequence);
 
 	/**
 	 * Add a new binding to this sequence that will possess the specified object
@@ -402,7 +402,7 @@ public:
 	 * @return A unique identifier for the new binding
 	 */
 	UFUNCTION(BlueprintCallable, Category="Sequencer|Sequence", meta=(ScriptMethod))
-	static FSequencerBindingProxy AddPossessable(UMovieSceneSequence* Sequence, UObject* ObjectToPossess);
+	static FMovieSceneBindingProxy AddPossessable(UMovieSceneSequence* Sequence, UObject* ObjectToPossess);
 
 	/**
 	 * Add a new binding to this sequence that will spawn the specified object
@@ -412,7 +412,7 @@ public:
 	 * @return A unique identifier for the new binding
 	 */
 	UFUNCTION(BlueprintCallable, Category="Sequencer|Sequence", meta=(ScriptMethod))
-	static FSequencerBindingProxy AddSpawnableFromInstance(UMovieSceneSequence* Sequence, UObject* ObjectToSpawn);
+	static FMovieSceneBindingProxy AddSpawnableFromInstance(UMovieSceneSequence* Sequence, UObject* ObjectToSpawn);
 
 	/**
 	 * Add a new binding to this sequence that will spawn the specified object
@@ -422,7 +422,7 @@ public:
 	 * @return A unique identifier for the new binding
 	 */
 	UFUNCTION(BlueprintCallable, Category="Sequencer|Sequence", meta=(ScriptMethod))
-	static FSequencerBindingProxy AddSpawnableFromClass(UMovieSceneSequence* Sequence, UClass* ClassToSpawn);
+	static FMovieSceneBindingProxy AddSpawnableFromClass(UMovieSceneSequence* Sequence, UClass* ClassToSpawn);
 
 	/**
 	 * Locate all the objects that correspond to the specified object ID, using the specified context
@@ -433,7 +433,7 @@ public:
 	 * @return An array of all bound objects
 	 */
 	UFUNCTION(BlueprintCallable, Category="Sequencer|Sequence", meta=(ScriptMethod))
-	static TArray<UObject*> LocateBoundObjects(UMovieSceneSequence* Sequence, const FSequencerBindingProxy& InBinding, UObject* Context);
+	static TArray<UObject*> LocateBoundObjects(UMovieSceneSequence* Sequence, const FMovieSceneBindingProxy& InBinding, UObject* Context);
 
 
 	/**
@@ -446,7 +446,7 @@ public:
 	 */
 	UE_DEPRECATED(5.0, "Please migrate to GetBindingID or GetPortableBindingID depending on use-case.")
 	UFUNCTION(BlueprintCallable, Category = "Sequencer|Sequence", meta = (ScriptMethod, DeprecatedFunction, DeprecationMessage="Please migrate to GetBindingID or GetPortableBindingID depending on use-case."))
-	static FMovieSceneObjectBindingID MakeBindingID(UMovieSceneSequence* MasterSequence, const FSequencerBindingProxy& InBinding, EMovieSceneObjectBindingSpace Space = EMovieSceneObjectBindingSpace::Root);
+	static FMovieSceneObjectBindingID MakeBindingID(UMovieSceneSequence* MasterSequence, const FMovieSceneBindingProxy& InBinding, EMovieSceneObjectBindingSpace Space = EMovieSceneObjectBindingSpace::Root);
 
 
 	/**
@@ -457,7 +457,7 @@ public:
 	 * @return The binding's id
 	 */
 	UFUNCTION(BlueprintCallable, Category = "Sequencer|Sequence", meta = (ScriptMethod))
-	static FMovieSceneObjectBindingID GetBindingID(const FSequencerBindingProxy& InBinding);
+	static FMovieSceneObjectBindingID GetBindingID(const FMovieSceneBindingProxy& InBinding);
 
 
 	/**
@@ -471,7 +471,7 @@ public:
 	 * @return The binding's id
 	 */
 	UFUNCTION(BlueprintCallable, Category = "Sequencer|Sequence", meta = (ScriptMethod))
-	static FMovieSceneObjectBindingID GetPortableBindingID(UMovieSceneSequence* MasterSequence, UMovieSceneSequence* DestinationSequence, const FSequencerBindingProxy& InBinding);
+	static FMovieSceneObjectBindingID GetPortableBindingID(UMovieSceneSequence* MasterSequence, UMovieSceneSequence* DestinationSequence, const FMovieSceneBindingProxy& InBinding);
 
 	/**
 	 * Make a binding for the given binding ID
@@ -481,7 +481,7 @@ public:
 	 * @return The new binding proxy
 	 */
 	UFUNCTION(BlueprintCallable, Category = "Sequencer|Sequence", meta = (ScriptMethod))
-	static FSequencerBindingProxy ResolveBindingID(UMovieSceneSequence* MasterSequence, FMovieSceneObjectBindingID InObjectBindingID);
+	static FMovieSceneBindingProxy ResolveBindingID(UMovieSceneSequence* MasterSequence, FMovieSceneObjectBindingID InObjectBindingID);
 
 
 	/**
diff --git a/Engine/Plugins/MovieScene/SequencerScripting/Source/SequencerScripting/Public/SequencerBindingProxy.h b/Engine/Plugins/MovieScene/SequencerScripting/Source/SequencerScripting/Public/SequencerBindingProxy.h
--- a/Engine/Plugins/MovieScene/SequencerScripting/Source/SequencerScripting/Public/SequencerBindingProxy.h
+++ b/Engine/Plugins/MovieScene/SequencerScripting/Source/SequencerScripting/Public/SequencerBindingProxy.h
@@ -4,32 +4,7 @@
 
 #include "CoreTypes.h"
 #include "UObject/ObjectMacros.h"
+#include "MovieSceneBindingProxy.h"
 
-#include "SequencerBindingProxy.generated.h"
-
-class UMovieScene;
-class UMovieSceneSequence;
-
-
-USTRUCT(BlueprintType)
-struct FSequencerBindingProxy
-{
-	GENERATED_BODY()
-
-	FSequencerBindingProxy()
-		: Sequence(nullptr)
-	{}
-
-	FSequencerBindingProxy(const FGuid& InBindingID, UMovieSceneSequence* InSequence)
-		: BindingID(InBindingID)
-		, Sequence(InSequence)
-	{}
-
-	UMovieScene* GetMovieScene() const;
-
-	UPROPERTY()
-	FGuid BindingID;
-
-	UPROPERTY(BlueprintReadOnly, Category=Binding)
-	TObjectPtr<UMovieSceneSequence> Sequence;
-};
+UE_DEPRECATED(5.1, "Use FMovieSceneBindingProxy instead.") 
+typedef FMovieSceneBindingProxy FSequencerBindingProxy;
\ No newline at end of file
diff --git a/Engine/Plugins/MovieScene/SequencerScripting/Source/SequencerScriptingEditor/Private/SequencerTools.cpp b/Engine/Plugins/MovieScene/SequencerScripting/Source/SequencerScriptingEditor/Private/SequencerTools.cpp
--- a/Engine/Plugins/MovieScene/SequencerScripting/Source/SequencerScriptingEditor/Private/SequencerTools.cpp
+++ b/Engine/Plugins/MovieScene/SequencerScripting/Source/SequencerScriptingEditor/Private/SequencerTools.cpp
@@ -34,7 +34,6 @@
 #include "LevelSequenceAnimSequenceLink.h"
 #include "AnimSequenceLevelSequenceLink.h"
 
-
 #define LOCTEXT_NAMESPACE "SequencerTools"
 
 bool USequencerToolsFunctionLibrary::RenderMovie(UMovieSceneCapture* InCaptureSettings, FOnRenderMovieStopped OnFinishedCallback)
@@ -118,7 +117,7 @@ void USequencerToolsFunctionLibrary::CancelMovieRender()
 	}
 }
 
-TArray<FSequencerBoundObjects> USequencerToolsFunctionLibrary::GetBoundObjects(UWorld* InWorld, ULevelSequence* InSequence, const TArray<FSequencerBindingProxy>& InBindings, const FSequencerScriptingRange& InRange)
+TArray<FSequencerBoundObjects> USequencerToolsFunctionLibrary::GetBoundObjects(UWorld* InWorld, ULevelSequence* InSequence, const TArray<FMovieSceneBindingProxy>& InBindings, const FSequencerScriptingRange& InRange)
 {
 	ALevelSequenceActor* OutActor;
 	FMovieSceneSequencePlaybackSettings Settings;
@@ -137,7 +136,7 @@ TArray<FSequencerBoundObjects> USequencerToolsFunctionLibrary::GetBoundObjects(U
 	FMovieSceneSequenceID SequenceId = Player->State.FindSequenceId(InSequence);
 
 	TArray<FSequencerBoundObjects> BoundObjects;
-	for (FSequencerBindingProxy Binding : InBindings)
+	for (FMovieSceneBindingProxy Binding : InBindings)
 	{
 		FMovieSceneObjectBindingID ObjectBinding = UE::MovieScene::FFixedObjectBindingID(Binding.BindingID, SequenceId);
 		BoundObjects.Add(FSequencerBoundObjects(Binding, Player->GetBoundObjects(ObjectBinding)));
@@ -172,7 +171,7 @@ TArray<FSequencerBoundObjects> USequencerToolsFunctionLibrary::GetObjectBindings
 		for (FMovieSceneObjectBindingID ObjectBinding : ObjectBindings)
 		{
 			FMovieSceneSequenceID SequenceID = ObjectBinding.ResolveSequenceID(MovieSceneSequenceID::Root, *Player);
-			FSequencerBindingProxy Binding(ObjectBinding.GetGuid(), Player->State.FindSequence(SequenceID));
+			FMovieSceneBindingProxy Binding(ObjectBinding.GetGuid(), Player->State.FindSequence(SequenceID));
 			BoundObjects.Add(FSequencerBoundObjects(Binding, TArray<UObject*>({ Object })));
 		}
 	}
@@ -182,7 +181,7 @@ TArray<FSequencerBoundObjects> USequencerToolsFunctionLibrary::GetObjectBindings
 	return BoundObjects;
 }
 
-void GatherDescendantBindings(const FSequencerBindingProxy& Binding, UMovieSceneSequence* Sequence, TArray<FSequencerBindingProxy>& AllBindings)
+void GatherDescendantBindings(const FMovieSceneBindingProxy& Binding, UMovieSceneSequence* Sequence, TArray<FMovieSceneBindingProxy>& AllBindings)
 {
 	UMovieScene* MovieScene = Sequence->GetMovieScene();
 
@@ -191,14 +190,14 @@ void GatherDescendantBindings(const FSequencerBindingProxy& Binding, UMovieScene
 		FMovieScenePossessable& Possessable = MovieScene->GetPossessable(Index);
 		if (Possessable.GetParent() == Binding.BindingID)
 		{
-			FSequencerBindingProxy ChildBinding(Possessable.GetGuid(), Sequence);
+			FMovieSceneBindingProxy ChildBinding(Possessable.GetGuid(), Sequence);
 			AllBindings.Add(ChildBinding);
 			GatherDescendantBindings(ChildBinding, Sequence, AllBindings);
 		}
 	}
 }
 
-bool ExportFBXInternal(UWorld* World, UMovieSceneSequence* Sequence, const TArray<FSequencerBindingProxy>& InBindings, const TArray<UMovieSceneTrack*>& MasterTracks, UFbxExportOption* OverrideOptions, const FString& InFBXFileName, UMovieSceneSequencePlayer* Player)
+bool ExportFBXInternal(UWorld* World, UMovieSceneSequence* Sequence, const TArray<FMovieSceneBindingProxy>& InBindings, const TArray<UMovieSceneTrack*>& MasterTracks, UFbxExportOption* OverrideOptions, const FString& InFBXFileName, UMovieSceneSequencePlayer* Player)
 {
 	UnFbx::FFbxExporter* Exporter = UnFbx::FFbxExporter::GetInstance();
 	//Show the fbx export dialog options
@@ -206,8 +205,8 @@ bool ExportFBXInternal(UWorld* World, UMovieSceneSequence* Sequence, const TArra
 
 	UMovieScene* MovieScene = Sequence->GetMovieScene();
 
-	TArray<FSequencerBindingProxy> AllBindings;
-	for (const FSequencerBindingProxy& Binding : InBindings)
+	TArray<FMovieSceneBindingProxy> AllBindings;
+	for (const FMovieSceneBindingProxy& Binding : InBindings)
 	{
 		AllBindings.Add(Binding);
 
@@ -215,7 +214,7 @@ bool ExportFBXInternal(UWorld* World, UMovieSceneSequence* Sequence, const TArra
 	}
 
 	TArray<FGuid> Bindings;
-	for (const FSequencerBindingProxy& Proxy : AllBindings)
+	for (const FMovieSceneBindingProxy& Proxy : AllBindings)
 	{
 		if (Proxy.Sequence == Sequence)
 		{
@@ -247,7 +246,7 @@ bool ExportFBXInternal(UWorld* World, UMovieSceneSequence* Sequence, const TArra
 	return bDidExport;
 }
 
-bool USequencerToolsFunctionLibrary::ExportLevelSequenceFBX(UWorld* World, ULevelSequence* Sequence, const TArray<FSequencerBindingProxy>& InBindings, const TArray<UMovieSceneTrack*>& InMasterTracks, UFbxExportOption* OverrideOptions, const FString& InFBXFileName)
+bool USequencerToolsFunctionLibrary::ExportLevelSequenceFBX(UWorld* World, ULevelSequence* Sequence, const TArray<FMovieSceneBindingProxy>& InBindings, const TArray<UMovieSceneTrack*>& InMasterTracks, UFbxExportOption* OverrideOptions, const FString& InFBXFileName)
 {
 	ALevelSequenceActor* OutActor;
 	FMovieSceneSequencePlaybackSettings Settings;
@@ -261,7 +260,7 @@ bool USequencerToolsFunctionLibrary::ExportLevelSequenceFBX(UWorld* World, ULeve
 	return bSuccess;
 }
 
-static USkeletalMeshComponent* GetSkelMeshComponent(IMovieScenePlayer* Player, const FSequencerBindingProxy& Binding)
+static USkeletalMeshComponent* GetSkelMeshComponent(IMovieScenePlayer* Player, const FMovieSceneBindingProxy& Binding)
 {
 	FMovieSceneSequenceIDRef Template = MovieSceneSequenceID::Root;
 	for (TWeakObjectPtr<UObject> RuntimeObject : Player->FindBoundObjects(Binding.BindingID, Template))
@@ -288,7 +287,7 @@ static USkeletalMeshComponent* GetSkelMeshComponent(IMovieScenePlayer* Player, c
 	return nullptr;
 }
 
-bool USequencerToolsFunctionLibrary::ExportAnimSequence(UWorld* World, ULevelSequence*  Sequence,  UAnimSequence* AnimSequence, UAnimSeqExportOption* ExportOptions,const FSequencerBindingProxy& Binding, bool bCreateLink)
+bool USequencerToolsFunctionLibrary::ExportAnimSequence(UWorld* World, ULevelSequence*  Sequence,  UAnimSequence* AnimSequence, UAnimSeqExportOption* ExportOptions,const FMovieSceneBindingProxy& Binding, bool bCreateLink)
 {
 	UMovieScene* MovieScene = Sequence->GetMovieScene();
 	if (Binding.Sequence != Sequence || !AnimSequence)
@@ -335,7 +334,7 @@ bool USequencerToolsFunctionLibrary::ExportAnimSequence(UWorld* World, ULevelSeq
 	return bResult;
 }
 
-bool USequencerToolsFunctionLibrary::LinkAnimSequence(ULevelSequence*  Sequence,  UAnimSequence* AnimSequence, const UAnimSeqExportOption* ExportOptions,const FSequencerBindingProxy& Binding)
+bool USequencerToolsFunctionLibrary::LinkAnimSequence(ULevelSequence*  Sequence,  UAnimSequence* AnimSequence, const UAnimSeqExportOption* ExportOptions,const FMovieSceneBindingProxy& Binding)
 {
 	if (!Sequence || !AnimSequence || !ExportOptions || Binding.Sequence != Sequence)
 	{
@@ -557,7 +556,7 @@ void ImportFBXCamera(UnFbx::FFbxImporter* FbxImporter, UWorld* World, UMovieScen
 	MovieSceneToolHelpers::ImportFBXCameraToExisting(FbxImporter, Sequence, Player, TemplateID, InObjectBindingMap, bMatchByNameOnly, true);
 }
 
-bool ImportFBXInternal(UWorld* World, UMovieSceneSequence* Sequence, const TArray<FSequencerBindingProxy>& InBindings, UMovieSceneUserImportFBXSettings* ImportFBXSettings, const FString&  ImportFilename, UMovieSceneSequencePlayer* Player)
+bool ImportFBXInternal(UWorld* World, UMovieSceneSequence* Sequence, const TArray<FMovieSceneBindingProxy>& InBindings, UMovieSceneUserImportFBXSettings* ImportFBXSettings, const FString&  ImportFilename, UMovieSceneSequencePlayer* Player)
 {
 	UMovieScene* MovieScene = Sequence->GetMovieScene();
 	if (!MovieScene || MovieScene->IsReadOnly())
@@ -565,8 +564,8 @@ bool ImportFBXInternal(UWorld* World, UMovieSceneSequence* Sequence, const TArra
 		return false;
 	}
 
-	TArray<FSequencerBindingProxy> AllBindings;
-	for (const FSequencerBindingProxy& Binding : InBindings)
+	TArray<FMovieSceneBindingProxy> AllBindings;
+	for (const FMovieSceneBindingProxy& Binding : InBindings)
 	{
 		AllBindings.Add(Binding);
 
@@ -574,7 +573,7 @@ bool ImportFBXInternal(UWorld* World, UMovieSceneSequence* Sequence, const TArra
 	}
 
 	TMap<FGuid, FString> ObjectBindingMap;
-	for (const FSequencerBindingProxy& Binding : AllBindings)
+	for (const FMovieSceneBindingProxy& Binding : AllBindings)
 	{
 		FString Name = MovieScene->GetObjectDisplayName(Binding.BindingID).ToString();
 		ObjectBindingMap.Add(Binding.BindingID, Name);
@@ -611,7 +610,7 @@ bool ImportFBXInternal(UWorld* World, UMovieSceneSequence* Sequence, const TArra
 	return bResult;
 }
 
-bool USequencerToolsFunctionLibrary::ImportLevelSequenceFBX(UWorld* World, ULevelSequence* Sequence, const TArray<FSequencerBindingProxy>& InBindings, UMovieSceneUserImportFBXSettings* ImportFBXSettings, const FString&  ImportFilename)
+bool USequencerToolsFunctionLibrary::ImportLevelSequenceFBX(UWorld* World, ULevelSequence* Sequence, const TArray<FMovieSceneBindingProxy>& InBindings, UMovieSceneUserImportFBXSettings* ImportFBXSettings, const FString&  ImportFilename)
 {
 	ALevelSequenceActor* OutActor;
 	FMovieSceneSequencePlaybackSettings Settings;
diff --git a/Engine/Plugins/MovieScene/SequencerScripting/Source/SequencerScriptingEditor/Public/SequencerTools.h b/Engine/Plugins/MovieScene/SequencerScripting/Source/SequencerScriptingEditor/Public/SequencerTools.h
--- a/Engine/Plugins/MovieScene/SequencerScripting/Source/SequencerScriptingEditor/Public/SequencerTools.h
+++ b/Engine/Plugins/MovieScene/SequencerScripting/Source/SequencerScriptingEditor/Public/SequencerTools.h
@@ -8,7 +8,7 @@
 #include "Modules/ModuleManager.h"
 #include "MovieSceneCaptureDialogModule.h"
 #include "Channels/MovieSceneEvent.h"
-#include "SequencerBindingProxy.h"
+#include "MovieSceneBindingProxy.h"
 #include "SequencerScriptingRange.h"
 #include "SequencerTools.generated.h"
 
@@ -38,13 +38,13 @@ struct SEQUENCERSCRIPTINGEDITOR_API FSequencerBoundObjects
 	FSequencerBoundObjects()
 	{}
 
-	FSequencerBoundObjects(FSequencerBindingProxy InBindingProxy, const TArray<UObject*>& InBoundObjects)
+	FSequencerBoundObjects(FMovieSceneBindingProxy InBindingProxy, const TArray<UObject*>& InBoundObjects)
 		: BindingProxy(InBindingProxy)
 		, BoundObjects(InBoundObjects)
 	{}
 
 	UPROPERTY(BlueprintReadWrite, Category=Binding)
-	FSequencerBindingProxy BindingProxy;
+	FMovieSceneBindingProxy BindingProxy;
 
 	UPROPERTY(BlueprintReadWrite, Category=Binding)
 	TArray<TObjectPtr<UObject>> BoundObjects;
@@ -109,7 +109,7 @@ public:
 	 * which allows for retrieving spawnables in that period of time.
 	 */
 	UFUNCTION(BlueprintCallable, Category = "Editor Scripting | Sequencer Tools")
-	static TArray<FSequencerBoundObjects> GetBoundObjects(UWorld* InWorld, ULevelSequence* InSequence, const TArray<FSequencerBindingProxy>& InBindings, const FSequencerScriptingRange& InRange);
+	static TArray<FSequencerBoundObjects> GetBoundObjects(UWorld* InWorld, ULevelSequence* InSequence, const TArray<FMovieSceneBindingProxy>& InBindings, const FSequencerScriptingRange& InRange);
 
 	/*
 	 * Get the object bindings for the requested object. The sequence will be evaluated in lower bound of the specified range, 
@@ -130,7 +130,7 @@ public:
 	 * @InFBXFileName File to create
 	 */
 	UFUNCTION(BlueprintCallable, Category = "Editor Scripting | Sequencer Tools | FBX")
-	static bool ExportLevelSequenceFBX(UWorld* InWorld, ULevelSequence* InSequence, const TArray<FSequencerBindingProxy>& InBindings, const TArray<UMovieSceneTrack*>& InMasterTracks, UFbxExportOption* OverrideOptions, const FString& InFBXFileName);
+	static bool ExportLevelSequenceFBX(UWorld* InWorld, ULevelSequence* InSequence, const TArray<FMovieSceneBindingProxy>& InBindings, const TArray<UMovieSceneTrack*>& InMasterTracks, UFbxExportOption* OverrideOptions, const FString& InFBXFileName);
 
 	/*
 	 * Export Passed in Binding as an Anim Seqquence.
@@ -144,7 +144,7 @@ public:
 	 * @bCreateLink If true will create a link between the animation sequence and the level sequence
 	 */
 	UFUNCTION(BlueprintCallable, Category = "Editor Scripting | Sequencer Tools | Animation")
-	static bool ExportAnimSequence(UWorld* World, ULevelSequence*  Sequence, UAnimSequence* AnimSequence, UAnimSeqExportOption* ExportOption, const FSequencerBindingProxy& Binding, bool bCreateLink);
+	static bool ExportAnimSequence(UWorld* World, ULevelSequence*  Sequence, UAnimSequence* AnimSequence, UAnimSeqExportOption* ExportOption, const FMovieSceneBindingProxy& Binding, bool bCreateLink);
 
 	/*
 	 * Links a LevelSequence's SkeletalMesh binding to an existing anim sequence.
@@ -155,7 +155,7 @@ public:
 	 * @InBinding Binding that has a skelmesh component on it
 	 */
 	UFUNCTION(BlueprintCallable, Category = "Editor Scripting | Sequencer Tools | Animation")
-	static bool LinkAnimSequence(ULevelSequence* Sequence, UAnimSequence* AnimSequence, const UAnimSeqExportOption* ExportOptions, const FSequencerBindingProxy& Binding);
+	static bool LinkAnimSequence(ULevelSequence* Sequence, UAnimSequence* AnimSequence, const UAnimSeqExportOption* ExportOptions, const FMovieSceneBindingProxy& Binding);
 
 	/*
 	 * Get the link to the level sequence if it exists on this anim sequence
@@ -186,7 +186,7 @@ public:
 	 * @InPlayer Player to bind to
 	 */
 	UFUNCTION(BlueprintCallable, Category = "Editor Scripting | Sequencer Tools | FBX")
-	static bool ImportLevelSequenceFBX(UWorld* InWorld, ULevelSequence* InSequence, const TArray<FSequencerBindingProxy>& InBindings, UMovieSceneUserImportFBXSettings* InImportFBXSettings, const FString& InImportFilename);
+	static bool ImportLevelSequenceFBX(UWorld* InWorld, ULevelSequence* InSequence, const TArray<FMovieSceneBindingProxy>& InBindings, UMovieSceneUserImportFBXSettings* InImportFBXSettings, const FString& InImportFilename);
 
 public:
 	/**
diff --git a/Engine/Source/Runtime/MovieScene/Private/MovieSceneBindingProxy.cpp b/Engine/Source/Runtime/MovieScene/Private/MovieSceneBindingProxy.cpp
new file mode 100644
--- /dev/null
+++ b/Engine/Source/Runtime/MovieScene/Private/MovieSceneBindingProxy.cpp
@@ -0,0 +1,9 @@
+// Copyright Epic Games, Inc. All Rights Reserved.
+
+#include "MovieSceneBindingProxy.h"
+#include "MovieSceneSequence.h"
+
+UMovieScene* FMovieSceneBindingProxy::GetMovieScene() const
+{
+	return Sequence ? Sequence->GetMovieScene() : nullptr;
+}
\ No newline at end of file
diff --git a/Engine/Source/Runtime/MovieScene/Public/MovieSceneBindingProxy.h b/Engine/Source/Runtime/MovieScene/Public/MovieSceneBindingProxy.h
new file mode 100644
--- /dev/null
+++ b/Engine/Source/Runtime/MovieScene/Public/MovieSceneBindingProxy.h
@@ -0,0 +1,39 @@
+// Copyright Epic Games, Inc. All Rights Reserved.
+
+#pragma once
+
+#include "CoreTypes.h"
+#include "UObject/ObjectMacros.h"
+
+#include "MovieSceneBindingProxy.generated.h"
+
+class UMovieScene;
+class UMovieSceneSequence;
+
+//
+// Movie Scene Binding Proxy represents the binding ID (an FGuid) and the corresponding sequence that it relates to. 
+// This is primarily used for scripting where there is no support for FMovieSceneSequenceID and use cases where 
+// managing multiple bindings with their corresponding sequences is necessary.
+//
+USTRUCT(BlueprintType)
+struct MOVIESCENE_API FMovieSceneBindingProxy
+{
+	GENERATED_BODY()
+
+	FMovieSceneBindingProxy()
+		: Sequence(nullptr)
+	{}
+
+	FMovieSceneBindingProxy(const FGuid& InBindingID, UMovieSceneSequence* InSequence)
+		: BindingID(InBindingID)
+		, Sequence(InSequence)
+	{}
+
+	UMovieScene* GetMovieScene() const;
+
+	UPROPERTY(BlueprintReadOnly, Category=Binding)
+	FGuid BindingID;
+
+	UPROPERTY(BlueprintReadOnly, Category=Binding)
+	TObjectPtr<UMovieSceneSequence> Sequence;
+};
