Sequencer: Use the FLevelSequenceNodeNameAdapter for python export otherwise the actor label doesn't come through the export

diff --git a/Engine/Plugins/MovieScene/SequencerScripting/Source/SequencerScriptingEditor/Private/SequencerTools.cpp b/Engine/Plugins/MovieScene/SequencerScripting/Source/SequencerScriptingEditor/Private/SequencerTools.cpp
--- a/Engine/Plugins/MovieScene/SequencerScripting/Source/SequencerScriptingEditor/Private/SequencerTools.cpp
+++ b/Engine/Plugins/MovieScene/SequencerScripting/Source/SequencerScriptingEditor/Private/SequencerTools.cpp
@@ -222,9 +222,9 @@ bool ExportFBXInternal(UWorld* World, UMovieSceneSequence* Sequence, const TArra
 		}
 	}
 
-	INodeNameAdapter NodeNameAdapter;
-	Player->State.AssignSequence(MovieSceneSequenceID::Root, *Sequence, *Player);
 	FMovieSceneSequenceIDRef Template = MovieSceneSequenceID::Root;
+	UnFbx::FFbxExporter::FLevelSequenceNodeNameAdapter NodeNameAdapter(MovieScene, Player, Template);
+	Player->State.AssignSequence(Template, *Sequence, *Player);
 	bool bDidExport = false;
 	FMovieSceneSequenceTransform RootToLocalTransform;
 	
@@ -294,7 +294,7 @@ bool USequencerToolsFunctionLibrary::ExportAnimSequence(UWorld* World, ULevelSeq
 	{
 		return false;
 	}
-	INodeNameAdapter NodeNameAdapter;
+
 	ALevelSequenceActor* OutActor;
 	FMovieSceneSequencePlaybackSettings Settings;
 	FLevelSequenceCameraSettings CameraSettings;
