Sequencer: Add meta=DeprecatedFunction to GetChannels and FindChannelsByType

diff --git a/Engine/Plugins/MovieScene/SequencerScripting/Source/SequencerScripting/Public/ExtensionLibraries/MovieSceneSectionExtensions.h b/Engine/Plugins/MovieScene/SequencerScripting/Source/SequencerScripting/Public/ExtensionLibraries/MovieSceneSectionExtensions.h
--- a/Engine/Plugins/MovieScene/SequencerScripting/Source/SequencerScripting/Public/ExtensionLibraries/MovieSceneSectionExtensions.h
+++ b/Engine/Plugins/MovieScene/SequencerScripting/Source/SequencerScripting/Public/ExtensionLibraries/MovieSceneSectionExtensions.h
@@ -217,8 +217,8 @@ public:
 	* @param Section       The section to use.
 	* @return An array containing any key channels that match the type specified
 	*/
-	UE_DEPRECATED(5.0, "Please use GetAllChannels instead")
-	UFUNCTION(BlueprintCallable, Category = "Sequencer|Section", meta=(ScriptMethod, DevelopmentOnly))
+	UE_DEPRECATED(5.0, "Please use GetAllChannels instead and query double-precision channels as such where appropriate")
+	UFUNCTION(BlueprintCallable, Category = "Sequencer|Section", meta=(ScriptMethod, DevelopmentOnly, DeprecatedFunction, DeprecationMessage = "Please use GetAllChannels instead and query double-precision channels as such where appropriate"))
 	static TArray<UMovieSceneScriptingChannel*> GetChannels(UMovieSceneSection* Section);
 
 	/**
@@ -244,7 +244,7 @@ public:
 	* @return An array containing any key channels that match the type specified
 	*/
 	UE_DEPRECATED(5.0, "Please use GetChannelsByType instead")
-	UFUNCTION(BlueprintCallable, Category = "Sequencer|Section", meta = (ScriptMethod, DeterminesOutputType="TrackType", DevelopmentOnly))
+	UFUNCTION(BlueprintCallable, Category = "Sequencer|Section", meta = (ScriptMethod, DeterminesOutputType="TrackType", DevelopmentOnly, DeprecatedFunction, DeprecationMessage = "Please use GetChannelsByType instead"))
 	static TArray<UMovieSceneScriptingChannel*> FindChannelsByType(UMovieSceneSection* Section, TSubclassOf<UMovieSceneScriptingChannel> ChannelType);
 
 	/**
