Sequencer: Fix export not evaluating the start time in display rate (instead it was evaluating in tick resolution)

diff --git a/Engine/Plugins/Animation/ControlRig/Source/ControlRigEditor/Private/ControlRigSequencerEditorLibrary.cpp b/Engine/Plugins/Animation/ControlRig/Source/ControlRigEditor/Private/ControlRigSequencerEditorLibrary.cpp
--- a/Engine/Plugins/Animation/ControlRig/Source/ControlRigEditor/Private/ControlRigSequencerEditorLibrary.cpp
+++ b/Engine/Plugins/Animation/ControlRig/Source/ControlRigEditor/Private/ControlRigSequencerEditorLibrary.cpp
@@ -659,7 +659,8 @@ bool UControlRigSequencerEditorLibrary::BakeToControlRig(UWorld* World, ULevelSe
 		if (LevelPlayer && SpawnableRestoreState.bWasChanged)
 		{
 			// Evaluate at the beginning of the subscene time to ensure that spawnables are created before export
-			LevelPlayer->SetPlaybackPosition(FMovieSceneSequencePlaybackParams(UE::MovieScene::DiscreteInclusiveLower(MovieScene->GetPlaybackRange()).Value, EUpdatePositionMethod::Play));
+			FFrameTime StartTime = FFrameRate::TransformTime(UE::MovieScene::DiscreteInclusiveLower(MovieScene->GetPlaybackRange()).Value, MovieScene->GetTickResolution(), MovieScene->GetDisplayRate());
+			LevelPlayer->SetPlaybackPosition(FMovieSceneSequencePlaybackParams(StartTime, EUpdatePositionMethod::Play));
 		}
 		TArrayView<TWeakObjectPtr<>>  Result = Player->FindBoundObjects(Binding.BindingID, Template);
 	
diff --git a/Engine/Plugins/Importers/USDImporter/Source/USDStage/Private/USDLevelSequenceHelper.cpp b/Engine/Plugins/Importers/USDImporter/Source/USDStage/Private/USDLevelSequenceHelper.cpp
--- a/Engine/Plugins/Importers/USDImporter/Source/USDStage/Private/USDLevelSequenceHelper.cpp
+++ b/Engine/Plugins/Importers/USDImporter/Source/USDStage/Private/USDLevelSequenceHelper.cpp
@@ -288,7 +288,8 @@ namespace UsdLevelSequenceHelperImpl
 			LevelPlayer->State.AssignSequence( MovieSceneSequenceID::Root, *LevelSequence, *Player );
 
 			// Evaluate at the beginning of the subscene time to ensure that spawnables are created before export
-			LevelPlayer->SetPlaybackPosition( FMovieSceneSequencePlaybackParams( UE::MovieScene::DiscreteInclusiveLower( MovieScene->GetPlaybackRange() ).Value, EUpdatePositionMethod::Play ) );
+			FFrameTime StartTime = FFrameRate::TransformTime(UE::MovieScene::DiscreteInclusiveLower(MovieScene->GetPlaybackRange()).Value, MovieScene->GetTickResolution(), MovieScene->GetDisplayRate());
+			LevelPlayer->SetPlaybackPosition( FMovieSceneSequencePlaybackParams( StartTime, EUpdatePositionMethod::Play ) );
 		}
 
 		MovieScene->Modify();
@@ -3081,9 +3082,10 @@ void FUsdLevelSequenceHelperImpl::HandleControlRigSectionChange( UMovieSceneCont
 			// Evaluate at the beginning of the subscene time to ensure that spawnables are created before export
 			// Note that we never actually generate spawnables on our LevelSequence, but its a common pattern to
 			// do this and the user may have added them manually
+			FFrameTime StartTime = FFrameRate::TransformTime(UE::MovieScene::DiscreteInclusiveLower(MovieScene->GetPlaybackRange()).Value, MovieScene->GetTickResolution(), MovieScene->GetDisplayRate());
 			LevelPlayer->SetPlaybackPosition(
 				FMovieSceneSequencePlaybackParams(
-					UE::MovieScene::DiscreteInclusiveLower( MovieScene->GetPlaybackRange() ).Value,
+					StartTime,
 					EUpdatePositionMethod::Play
 				)
 			);
diff --git a/Engine/Plugins/MovieScene/SequencerScripting/Source/SequencerScriptingEditor/Private/SequencerTools.cpp b/Engine/Plugins/MovieScene/SequencerScripting/Source/SequencerScriptingEditor/Private/SequencerTools.cpp
--- a/Engine/Plugins/MovieScene/SequencerScripting/Source/SequencerScriptingEditor/Private/SequencerTools.cpp
+++ b/Engine/Plugins/MovieScene/SequencerScripting/Source/SequencerScriptingEditor/Private/SequencerTools.cpp
@@ -234,7 +234,8 @@ bool ExportFBXInternal(UWorld* World, UMovieSceneSequence* Sequence, const TArra
 		if (SpawnableRestoreState.bWasChanged)
 		{
 			// Evaluate at the beginning of the subscene time to ensure that spawnables are created before export
-			Player->SetPlaybackPosition(FMovieSceneSequencePlaybackParams(UE::MovieScene::DiscreteInclusiveLower(MovieScene->GetPlaybackRange()).Value, EUpdatePositionMethod::Play));
+			FFrameTime StartTime = FFrameRate::TransformTime(UE::MovieScene::DiscreteInclusiveLower(MovieScene->GetPlaybackRange()).Value, MovieScene->GetTickResolution(), MovieScene->GetDisplayRate());
+			Player->SetPlaybackPosition(FMovieSceneSequencePlaybackParams(StartTime, EUpdatePositionMethod::Play));
 		}
 
 		bDidExport = MovieSceneToolHelpers::ExportFBX(World, MovieScene, Player, Bindings, MasterTracks, NodeNameAdapter, Template, InFBXFileName, RootToLocalTransform);
@@ -312,7 +313,8 @@ bool USequencerToolsFunctionLibrary::ExportAnimSequence(UWorld* World, ULevelSeq
 		if (SpawnableRestoreState.bWasChanged)
 		{
 			// Evaluate at the beginning of the subscene time to ensure that spawnables are created before export
-			Player->SetPlaybackPosition(FMovieSceneSequencePlaybackParams(UE::MovieScene::DiscreteInclusiveLower(MovieScene->GetPlaybackRange()).Value, EUpdatePositionMethod::Play));
+			FFrameTime StartTime = FFrameRate::TransformTime(UE::MovieScene::DiscreteInclusiveLower(MovieScene->GetPlaybackRange()).Value, MovieScene->GetTickResolution(), MovieScene->GetDisplayRate());
+			Player->SetPlaybackPosition(FMovieSceneSequencePlaybackParams(StartTime, EUpdatePositionMethod::Play));
 		}
  
 		USkeletalMeshComponent* SkeletalMeshComp =  GetSkelMeshComponent(Player, Binding);
@@ -623,7 +625,8 @@ bool ImportFBXInternal(UWorld* World, UMovieSceneSequence* Sequence, const TArra
 		if (SpawnableRestoreState.bWasChanged)
 		{
 			// Evaluate at the beginning of the subscene time to ensure that spawnables are created before export
-			Player->SetPlaybackPosition(FMovieSceneSequencePlaybackParams(UE::MovieScene::DiscreteInclusiveLower(MovieScene->GetPlaybackRange()).Value, EUpdatePositionMethod::Play));
+			FFrameTime StartTime = FFrameRate::TransformTime(UE::MovieScene::DiscreteInclusiveLower(MovieScene->GetPlaybackRange()).Value, MovieScene->GetTickResolution(), MovieScene->GetDisplayRate());
+			Player->SetPlaybackPosition(FMovieSceneSequencePlaybackParams(StartTime, EUpdatePositionMethod::Play));
 		}
 
 		ImportFBXCamera(FbxImporter, World, Sequence, MovieScene, Player, MovieSceneSequenceID::Root, ObjectBindingMap, bMatchByNameOnly, ImportFBXSettings->bCreateCameras);
