Sequencer: ExportLevelSequenceFBX is now able to export a subsequence within a master/root sequence.

Changed to take params FSequencerExportFBXParams, deprecated the old ExportLevelSequenceFBX.

diff --git a/Engine/Plugins/MovieScene/SequencerScripting/Source/SequencerScriptingEditor/Private/SequencerTools.cpp b/Engine/Plugins/MovieScene/SequencerScripting/Source/SequencerScriptingEditor/Private/SequencerTools.cpp
--- a/Engine/Plugins/MovieScene/SequencerScripting/Source/SequencerScriptingEditor/Private/SequencerTools.cpp
+++ b/Engine/Plugins/MovieScene/SequencerScripting/Source/SequencerScriptingEditor/Private/SequencerTools.cpp
@@ -34,6 +34,8 @@
 #include "LevelSequenceAnimSequenceLink.h"
 #include "AnimSequenceLevelSequenceLink.h"
 
+#include "Compilation/MovieSceneCompiledDataManager.h"
+
 #define LOCTEXT_NAMESPACE "SequencerTools"
 
 bool USequencerToolsFunctionLibrary::RenderMovie(UMovieSceneCapture* InCaptureSettings, FOnRenderMovieStopped OnFinishedCallback)
@@ -197,8 +199,16 @@ void GatherDescendantBindings(const FMovieSceneBindingProxy& Binding, UMovieScen
 	}
 }
 
-bool ExportFBXInternal(UWorld* World, UMovieSceneSequence* Sequence, const TArray<FMovieSceneBindingProxy>& InBindings, const TArray<UMovieSceneTrack*>& MasterTracks, UFbxExportOption* OverrideOptions, const FString& InFBXFileName, UMovieSceneSequencePlayer* Player)
+bool ExportFBXInternal(const FSequencerExportFBXParams& InParams, UMovieSceneSequencePlayer* Player)
 {
+	UWorld* World = InParams.World;
+	UMovieSceneSequence* Sequence = InParams.Sequence;
+	UMovieSceneSequence* RootSequence = InParams.RootSequence;
+	TArray<FMovieSceneBindingProxy> BindingProxies = InParams.Bindings;
+	TArray<UMovieSceneTrack*> MasterTracks = InParams.MasterTracks;
+	UFbxExportOption* OverrideOptions = InParams.OverrideOptions;
+	FString FBXFileName = InParams.FBXFileName;
+
 	UnFbx::FFbxExporter* Exporter = UnFbx::FFbxExporter::GetInstance();
 	//Show the fbx export dialog options
 	Exporter->SetExportOptionsOverride(OverrideOptions);
@@ -206,7 +216,7 @@ bool ExportFBXInternal(UWorld* World, UMovieSceneSequence* Sequence, const TArra
 	UMovieScene* MovieScene = Sequence->GetMovieScene();
 
 	TArray<FMovieSceneBindingProxy> AllBindings;
-	for (const FMovieSceneBindingProxy& Binding : InBindings)
+	for (const FMovieSceneBindingProxy& Binding : BindingProxies)
 	{
 		AllBindings.Add(Binding);
 
@@ -222,11 +232,25 @@ bool ExportFBXInternal(UWorld* World, UMovieSceneSequence* Sequence, const TArra
 		}
 	}
 
-	FMovieSceneSequenceIDRef Template = MovieSceneSequenceID::Root;
+	Player->State.AssignSequence(MovieSceneSequenceID::Root, *RootSequence, *Player);
+	FMovieSceneSequenceIDRef Template = Player->State.FindSequenceId(Sequence);
 	UnFbx::FFbxExporter::FLevelSequenceNodeNameAdapter NodeNameAdapter(MovieScene, Player, Template);
-	Player->State.AssignSequence(Template, *Sequence, *Player);
-	bool bDidExport = false;
+
+	FMovieSceneSequenceHierarchy Hierarchy = FMovieSceneSequenceHierarchy();
+	UMovieSceneCompiledDataManager::CompileHierarchy(RootSequence, &Hierarchy, EMovieSceneServerClientMask::All);
+
+	const FMovieSceneSubSequenceData* SubSequenceData = Hierarchy.FindSubData(Template);
 	FMovieSceneSequenceTransform RootToLocalTransform;
+
+	FFrameTime StartTime = FFrameRate::TransformTime(UE::MovieScene::DiscreteInclusiveLower(MovieScene->GetPlaybackRange()).Value, MovieScene->GetTickResolution(), MovieScene->GetDisplayRate());
+
+	if (SubSequenceData)
+	{
+		RootToLocalTransform = SubSequenceData->RootToSequenceTransform;
+		StartTime = UE::MovieScene::DiscreteInclusiveLower(SubSequenceData->PlayRange.Value) * RootToLocalTransform.InverseLinearOnly();
+	}
+
+	bool bDidExport = false;
 	
 	{
 		FSpawnableRestoreState SpawnableRestoreState(MovieScene);
@@ -234,11 +258,10 @@ bool ExportFBXInternal(UWorld* World, UMovieSceneSequence* Sequence, const TArra
 		if (SpawnableRestoreState.bWasChanged)
 		{
 			// Evaluate at the beginning of the subscene time to ensure that spawnables are created before export
-			FFrameTime StartTime = FFrameRate::TransformTime(UE::MovieScene::DiscreteInclusiveLower(MovieScene->GetPlaybackRange()).Value, MovieScene->GetTickResolution(), MovieScene->GetDisplayRate());
 			Player->SetPlaybackPosition(FMovieSceneSequencePlaybackParams(StartTime, EUpdatePositionMethod::Play));
 		}
 
-		bDidExport = MovieSceneToolHelpers::ExportFBX(World, MovieScene, Player, Bindings, MasterTracks, NodeNameAdapter, Template, InFBXFileName, RootToLocalTransform);
+		bDidExport = MovieSceneToolHelpers::ExportFBX(World, MovieScene, Player, Bindings, MasterTracks, NodeNameAdapter, Template, FBXFileName, RootToLocalTransform);
 	}
 
 	Player->Stop();
@@ -247,16 +270,16 @@ bool ExportFBXInternal(UWorld* World, UMovieSceneSequence* Sequence, const TArra
 	return bDidExport;
 }
 
-bool USequencerToolsFunctionLibrary::ExportLevelSequenceFBX(UWorld* World, ULevelSequence* Sequence, const TArray<FMovieSceneBindingProxy>& InBindings, const TArray<UMovieSceneTrack*>& InMasterTracks, UFbxExportOption* OverrideOptions, const FString& InFBXFileName)
+bool USequencerToolsFunctionLibrary::ExportLevelSequenceFBX(const FSequencerExportFBXParams& InParams)
 {
 	ALevelSequenceActor* OutActor;
 	FMovieSceneSequencePlaybackSettings Settings;
 	FLevelSequenceCameraSettings CameraSettings;
-	ULevelSequencePlayer* Player = ULevelSequencePlayer::CreateLevelSequencePlayer(World, Sequence, Settings, OutActor);
-	Player->Initialize(Sequence, World->PersistentLevel, Settings, CameraSettings);
-	
-	bool bSuccess = ExportFBXInternal(World, Sequence, InBindings, InMasterTracks, OverrideOptions, InFBXFileName, Player);
-	World->DestroyActor(OutActor);
+	ULevelSequencePlayer* Player = ULevelSequencePlayer::CreateLevelSequencePlayer(InParams.World, InParams.RootSequence, Settings, OutActor);
+	Player->Initialize(InParams.RootSequence, InParams.World->PersistentLevel, Settings, CameraSettings);
+
+	bool bSuccess = ExportFBXInternal(InParams, Player);
+	InParams.World->DestroyActor(OutActor);
 
 	return bSuccess;
 }
diff --git a/Engine/Plugins/MovieScene/SequencerScripting/Source/SequencerScriptingEditor/Public/SequencerTools.h b/Engine/Plugins/MovieScene/SequencerScripting/Source/SequencerScriptingEditor/Public/SequencerTools.h
--- a/Engine/Plugins/MovieScene/SequencerScripting/Source/SequencerScriptingEditor/Public/SequencerTools.h
+++ b/Engine/Plugins/MovieScene/SequencerScripting/Source/SequencerScriptingEditor/Public/SequencerTools.h
@@ -68,6 +68,43 @@ struct FSequencerQuickBindingResult
 
 };
 
+USTRUCT(BlueprintType)
+struct FSequencerExportFBXParams
+{
+	GENERATED_BODY()
+
+	FSequencerExportFBXParams() {}
+	FSequencerExportFBXParams(UWorld* InWorld, ULevelSequence* InSequence, ULevelSequence* InRootSequence, const TArray<FMovieSceneBindingProxy>& InBindings, const TArray<UMovieSceneTrack*>& InMasterTracks, UFbxExportOption* InOverrideOptions, const FString& InFBXFileName)
+		: World(InWorld)
+		, Sequence(InSequence)
+		, RootSequence(InRootSequence)
+		, Bindings(InBindings)
+		, MasterTracks(InMasterTracks)
+		, OverrideOptions(InOverrideOptions)
+		, FBXFileName(InFBXFileName) {}
+	
+	UPROPERTY(BlueprintReadWrite, Category = "Movie Scene")
+	TObjectPtr<UWorld> World;
+	
+	UPROPERTY(BlueprintReadWrite, Category = "Movie Scene")
+	TObjectPtr<ULevelSequence> Sequence;
+	
+	UPROPERTY(BlueprintReadWrite, Category = "Movie Scene")
+	TObjectPtr<ULevelSequence> RootSequence;
+	
+	UPROPERTY(BlueprintReadWrite, Category = "Movie Scene")
+	TArray<FMovieSceneBindingProxy> Bindings;
+	
+	UPROPERTY(BlueprintReadWrite, Category = "Movie Scene")
+	TArray<TObjectPtr<UMovieSceneTrack>> MasterTracks;
+	
+	UPROPERTY(BlueprintReadWrite, Category = "Movie Scene")
+	TObjectPtr<UFbxExportOption> OverrideOptions;
+
+	UPROPERTY(BlueprintReadWrite, Category = "Movie Scene")
+	FString FBXFileName;
+};
+
 /** 
  * This is a set of helper functions to access various parts of the Sequencer API via Python. Because Sequencer itself is not suitable for exposing, most functionality
  * gets wrapped by UObjects that have an easier API to work with. This UObject provides access to these wrapper UObjects where needed. 
@@ -122,15 +159,12 @@ public:
 
 	/*
 	 * Export Passed in Bindings and Master Tracks to FBX
-	 *
-	 * @InWorld World to export
-	 * @InSequence Sequence to export
-	 * @InBindings Bindings to export
-	 * @InMasterTracks Master tracks to export
-	 * @InFBXFileName File to create
 	 */
 	UFUNCTION(BlueprintCallable, Category = "Editor Scripting | Sequencer Tools | FBX")
-	static bool ExportLevelSequenceFBX(UWorld* InWorld, ULevelSequence* InSequence, const TArray<FMovieSceneBindingProxy>& InBindings, const TArray<UMovieSceneTrack*>& InMasterTracks, UFbxExportOption* OverrideOptions, const FString& InFBXFileName);
+	static bool ExportLevelSequenceFBX(const FSequencerExportFBXParams& InParams);
+	
+	UE_DEPRECATED(5.1, "Please use ExportLevelSequenceFBX that takes a FSequencerExportFBXParams")
+	static bool ExportLevelSequenceFBX(UWorld* InWorld, ULevelSequence* InSequence, const TArray<FMovieSceneBindingProxy>& InBindings, const TArray<UMovieSceneTrack*>& InMasterTracks, UFbxExportOption* InOverrideOptions, const FString& InFBXFileName) { FSequencerExportFBXParams Params(InWorld, InSequence, InSequence, InBindings, InMasterTracks, InOverrideOptions, InFBXFileName); return ExportLevelSequenceFBX(Params); }
 
 	/*
 	 * Export Passed in Binding as an Anim Seqquence.
