Sequencer: Level Sequence Linking: When clearing the anim sequences from a level sequence, make sure it's pointing back to this level sequence before clearing asset user data
#trivial

diff --git a/Engine/Plugins/MovieScene/SequencerScripting/Source/SequencerScriptingEditor/Private/SequencerTools.cpp b/Engine/Plugins/MovieScene/SequencerScripting/Source/SequencerScriptingEditor/Private/SequencerTools.cpp
--- a/Engine/Plugins/MovieScene/SequencerScripting/Source/SequencerScriptingEditor/Private/SequencerTools.cpp
+++ b/Engine/Plugins/MovieScene/SequencerScripting/Source/SequencerScriptingEditor/Private/SequencerTools.cpp
@@ -374,7 +374,15 @@ void USequencerToolsFunctionLibrary::ClearLinkedAnimSequences(ULevelSequence* Le
 					{
 						if (IInterface_AssetUserData* AnimAssetUserData = Cast< IInterface_AssetUserData >(AnimSequence))
 						{
-							AnimAssetUserData->RemoveUserDataOfClass(UAnimSequenceLevelSequenceLink::StaticClass());
+							UAnimSequenceLevelSequenceLink* AnimLevelLink = AnimAssetUserData->GetAssetUserData< UAnimSequenceLevelSequenceLink >();
+							if (AnimLevelLink)
+							{
+								ULevelSequence* AnimLevelSequence = AnimLevelLink->ResolveLevelSequence();
+								if (AnimLevelSequence && AnimLevelSequence == LevelSequence)
+								{
+									AnimAssetUserData->RemoveUserDataOfClass(UAnimSequenceLevelSequenceLink::StaticClass());
+								}
+							}
 						}
 					}
 				}
