Sequencer: Separate application of UMovieSceneSequencePlayer::PlaybackSettings from Initialize

Captain's log of the reason behind why this is required:

When spawning a replicated Level Sequence Actor on the server, a SequencePlayer subobject is also created and replicated along with it. If the sequence player is interacted with directly (by way of calling things like PlayLooping, which adjust the loop count of the player's playback settings directly), there can be a justifiable discrepancy in playback settings on the server between the actor and its player. This is not a problem in practice because the Actor's playback settings are only used for initialization.

When these 2 objects are replicated to the client, the client receives the properties for both the actor and sequence player in the same network bunch, applies them, then calls PostNetReceive on both objects. Unfortunately it can call ALevelSequenceActor::PostNetReceive before UMovieSceneSequencePlayer::PostNetReceive, which means the Actor ends up calling Initialize on the player (because it hasn't been set up yet - it's only just been spawned by the net code), which ends up trampling over the player's PlaybackSettings that were just applied by the net replication.

Since the shadow state on the server thinks it has already sent those properties to the client, it never sends them again, so the client doesn't think it's looping when the server does.

diff --git a/Engine/Plugins/Animation/ControlRig/Source/ControlRigEditor/Private/ControlRigSequencerEditorLibrary.cpp b/Engine/Plugins/Animation/ControlRig/Source/ControlRigEditor/Private/ControlRigSequencerEditorLibrary.cpp
--- a/Engine/Plugins/Animation/ControlRig/Source/ControlRigEditor/Private/ControlRigSequencerEditorLibrary.cpp
+++ b/Engine/Plugins/Animation/ControlRig/Source/ControlRigEditor/Private/ControlRigSequencerEditorLibrary.cpp
@@ -92,9 +92,6 @@ static TArray<UObject*> GetBoundObjects(UWorld* World, ULevelSequence* LevelSequ
 	ULevelSequencePlayer* Player = ULevelSequencePlayer::CreateLevelSequencePlayer(World, LevelSequence, Settings, *OutActor);
 	*OutPlayer = Player;
 
-	Player->Initialize(LevelSequence, World->PersistentLevel, Settings, CameraSettings);
-	Player->State.AssignSequence(MovieSceneSequenceID::Root, *LevelSequence, *Player);
-
 	// Evaluation needs to occur in order to obtain spawnables
 	Player->SetPlaybackPosition(FMovieSceneSequencePlaybackParams(LevelSequence->GetMovieScene()->GetPlaybackRange().GetLowerBoundValue().Value, EUpdatePositionMethod::Play));
 
@@ -642,8 +639,6 @@ bool UControlRigSequencerEditorLibrary::BakeToControlRig(UWorld* World, ULevelSe
 	else
 	{
 		Player = LevelPlayer = ULevelSequencePlayer::CreateLevelSequencePlayer(World, LevelSequence, Settings, OutActor);
-		LevelPlayer->Initialize(LevelSequence, World->PersistentLevel, Settings, CameraSettings);
-		LevelPlayer->State.AssignSequence(MovieSceneSequenceID::Root, *LevelSequence, *Player);
 	}
 	if (Player == nullptr)
 	{
@@ -929,8 +924,6 @@ static bool GetControlRigValues(UWorld* World, ULevelSequence* LevelSequence, UC
 		FMovieSceneSequenceIDRef Template = MovieSceneSequenceID::Root;
 		FMovieSceneSequenceTransform RootToLocalTransform;
 		ULevelSequencePlayer* Player = ULevelSequencePlayer::CreateLevelSequencePlayer(World, LevelSequence, Settings, OutActor);
-		Player->Initialize(LevelSequence, World->PersistentLevel, Settings, CameraSettings);
-		Player->State.AssignSequence(MovieSceneSequenceID::Root, *LevelSequence, *Player);
 		return LocalGetControlRigControlValues(Player, LevelSequence, Template, RootToLocalTransform,
 			ControlRig, ControlName,TimeUnit, Frames, OutValues);
 
@@ -2120,8 +2113,6 @@ bool UControlRigSequencerEditorLibrary::ImportFBXToControlRigTrack(UWorld* World
 	FMovieSceneSequencePlaybackSettings Settings;
 	FLevelSequenceCameraSettings CameraSettings;
 	ULevelSequencePlayer* Player = ULevelSequencePlayer::CreateLevelSequencePlayer(World, Sequence, Settings, OutActor);
-	Player->Initialize(Sequence, World->GetLevel(0), Settings, CameraSettings);
-	Player->State.AssignSequence(MovieSceneSequenceID::Root, *Sequence, *Player);
 
 	INodeAndChannelMappings* ChannelMapping = Cast<INodeAndChannelMappings>(InTrack);
 	if (ChannelMapping)
diff --git a/Engine/Plugins/Animation/ControlRig/Source/ControlRigEditor/Private/Tools/ControlRigSnapper.cpp b/Engine/Plugins/Animation/ControlRig/Source/ControlRigEditor/Private/Tools/ControlRigSnapper.cpp
--- a/Engine/Plugins/Animation/ControlRig/Source/ControlRigEditor/Private/Tools/ControlRigSnapper.cpp
+++ b/Engine/Plugins/Animation/ControlRig/Source/ControlRigEditor/Private/Tools/ControlRigSnapper.cpp
@@ -170,8 +170,6 @@ bool FControlRigSnapper::GetControlRigControlTransforms(UWorld* World,ULevelSequ
 		FMovieSceneSequenceIDRef Template = MovieSceneSequenceID::Root;
 		FMovieSceneSequenceTransform RootToLocalTransform;
 		ULevelSequencePlayer* Player = ULevelSequencePlayer::CreateLevelSequencePlayer(World, LevelSequence, Settings, OutActor);
-		Player->Initialize(LevelSequence, World->PersistentLevel, Settings, CameraSettings);
-		Player->State.AssignSequence(MovieSceneSequenceID::Root, *LevelSequence, *Player);
 		TOptional<FFrameNumber> OptFrame;
 		bool Success = LocalGetControlRigControlTransforms(Player, OptFrame, LevelSequence, Template, RootToLocalTransform,
 			ControlRig, ControlName, Frames, ParentTransforms, OutTransforms);
diff --git a/Engine/Plugins/MovieScene/SequencerScripting/Source/SequencerScriptingEditor/Private/SequencerTools.cpp b/Engine/Plugins/MovieScene/SequencerScripting/Source/SequencerScriptingEditor/Private/SequencerTools.cpp
--- a/Engine/Plugins/MovieScene/SequencerScripting/Source/SequencerScriptingEditor/Private/SequencerTools.cpp
+++ b/Engine/Plugins/MovieScene/SequencerScripting/Source/SequencerScriptingEditor/Private/SequencerTools.cpp
@@ -129,9 +129,6 @@ TArray<FSequencerBoundObjects> USequencerToolsFunctionLibrary::GetBoundObjects(U
 
 	ULevelSequencePlayer* Player = ULevelSequencePlayer::CreateLevelSequencePlayer(InWorld, InSequence, Settings, OutActor);
 
-	Player->Initialize(InSequence, InWorld->PersistentLevel, Settings, CameraSettings);
-	Player->State.AssignSequence(MovieSceneSequenceID::Root, *InSequence, *Player);
-
 	// Evaluation needs to occur in order to obtain spawnables
 	FFrameRate Resolution = InSequence->GetMovieScene()->GetTickResolution();
 	TRange<FFrameNumber> SpecifiedRange = InRange.ToNative(Resolution);
@@ -160,9 +157,6 @@ TArray<FSequencerBoundObjects> USequencerToolsFunctionLibrary::GetObjectBindings
 
 	ULevelSequencePlayer* Player = ULevelSequencePlayer::CreateLevelSequencePlayer(InWorld, InSequence, Settings, OutActor);
 
-	Player->Initialize(InSequence, InWorld->PersistentLevel, Settings, CameraSettings);
-	Player->State.AssignSequence(MovieSceneSequenceID::Root, *InSequence, *Player);
-
 	FFrameRate Resolution = InSequence->GetMovieScene()->GetTickResolution();
 	TRange<FFrameNumber> SpecifiedRange = InRange.ToNative(Resolution);
 	Player->SetPlaybackPosition(FMovieSceneSequencePlaybackParams(SpecifiedRange.GetLowerBoundValue().Value, EUpdatePositionMethod::Play));
@@ -278,7 +272,6 @@ bool USequencerToolsFunctionLibrary::ExportLevelSequenceFBX(const FSequencerExpo
 	FMovieSceneSequencePlaybackSettings Settings;
 	FLevelSequenceCameraSettings CameraSettings;
 	ULevelSequencePlayer* Player = ULevelSequencePlayer::CreateLevelSequencePlayer(InParams.World, InParams.RootSequence, Settings, OutActor);
-	Player->Initialize(InParams.RootSequence, InParams.World->PersistentLevel, Settings, CameraSettings);
 
 	bool bSuccess = ExportFBXInternal(InParams, Player);
 	InParams.World->DestroyActor(OutActor);
@@ -327,8 +320,6 @@ bool USequencerToolsFunctionLibrary::ExportAnimSequence(UWorld* World, ULevelSeq
 	FMovieSceneSequenceIDRef Template = MovieSceneSequenceID::Root;
 	FMovieSceneSequenceTransform RootToLocalTransform;
 	ULevelSequencePlayer* Player = ULevelSequencePlayer::CreateLevelSequencePlayer(World, Sequence, Settings, OutActor);
-	Player->Initialize(Sequence, World->PersistentLevel, Settings, CameraSettings);
-	Player->State.AssignSequence(MovieSceneSequenceID::Root, *Sequence, *Player);
 
 	bool bResult = false;
 	
@@ -646,7 +637,6 @@ bool ImportFBXInternal(UWorld* World, UMovieSceneSequence* Sequence, const TArra
 	}
 
 	const bool bMatchByNameOnly = ImportFBXSettings->bMatchByNameOnly;
-	Player->State.AssignSequence(MovieSceneSequenceID::Root, *Sequence, *Player);
 
 	UnFbx::FFbxImporter* FbxImporter = UnFbx::FFbxImporter::GetInstance();
 	
@@ -677,7 +667,6 @@ bool USequencerToolsFunctionLibrary::ImportLevelSequenceFBX(UWorld* World, ULeve
 	FMovieSceneSequencePlaybackSettings Settings;
 	FLevelSequenceCameraSettings CameraSettings;
 	ULevelSequencePlayer* Player = ULevelSequencePlayer::CreateLevelSequencePlayer(World, Sequence, Settings, OutActor);
-	Player->Initialize(Sequence, World->GetLevel(0), Settings, CameraSettings);
 
 	bool bSuccess = ImportFBXInternal(World, Sequence, InBindings, ImportFBXSettings, ImportFilename, Player);
 	World->DestroyActor(OutActor);
@@ -707,8 +696,6 @@ bool USequencerToolsFunctionLibrary::ImportFBXToControlRig(UWorld* World, ULevel
 			FMovieSceneSequencePlaybackSettings Settings;
 			FLevelSequenceCameraSettings CameraSettings;
 			ULevelSequencePlayer* Player = ULevelSequencePlayer::CreateLevelSequencePlayer(World, Sequence, Settings, OutActor);
-			Player->Initialize(Sequence, World->GetLevel(0), Settings, CameraSettings);
-			Player->State.AssignSequence(MovieSceneSequenceID::Root, *Sequence, *Player);
 
 			const TArray<UMovieSceneTrack*>& Tracks = Binding.GetTracks();
 			TArray<FName> SelectedControls;
diff --git a/Engine/Source/Runtime/LevelSequence/Private/LevelSequenceActor.cpp b/Engine/Source/Runtime/LevelSequence/Private/LevelSequenceActor.cpp
--- a/Engine/Source/Runtime/LevelSequence/Private/LevelSequenceActor.cpp
+++ b/Engine/Source/Runtime/LevelSequence/Private/LevelSequenceActor.cpp
@@ -92,6 +92,7 @@ void ALevelSequenceActor::PostInitProperties()
 	// Have to initialize this here as any properties set on default subobjects inside the constructor
 	// Get stomped by the CDO's properties when the constructor exits.
 	SequencePlayer->SetPlaybackClient(this);
+	SequencePlayer->SetPlaybackSettings(PlaybackSettings);
 }
 
 void ALevelSequenceActor::RewindForReplay()
@@ -211,6 +212,8 @@ void ALevelSequenceActor::PostLoad()
 		bAutoPlay_DEPRECATED = false;
 	}
 
+	SequencePlayer->SetPlaybackSettings(PlaybackSettings);
+
 #if WITH_EDITORONLY_DATA
 	if (LevelSequence_DEPRECATED.IsValid())
 	{
@@ -257,7 +260,7 @@ void ALevelSequenceActor::SetSequence(ULevelSequence* InSequence)
 		// cbb: should ideally null out the template and player when no sequence is assigned, but that's currently not possible
 		if (InSequence)
 		{
-			SequencePlayer->Initialize(InSequence, GetLevel(), PlaybackSettings, CameraSettings);
+			SequencePlayer->Initialize(InSequence, GetLevel(), CameraSettings);
 		}
 	}
 }
@@ -269,7 +272,7 @@ void ALevelSequenceActor::InitializePlayer()
 		// Level sequence is already loaded. Initialize the player if it's not already initialized with this sequence
 		if (LevelSequenceAsset != SequencePlayer->GetSequence())
 		{
-			SequencePlayer->Initialize(LevelSequenceAsset, GetLevel(), PlaybackSettings, CameraSettings);
+			SequencePlayer->Initialize(LevelSequenceAsset, GetLevel(), CameraSettings);
 		}
 	}
 }
diff --git a/Engine/Source/Runtime/LevelSequence/Private/LevelSequencePlayer.cpp b/Engine/Source/Runtime/LevelSequence/Private/LevelSequencePlayer.cpp
--- a/Engine/Source/Runtime/LevelSequence/Private/LevelSequencePlayer.cpp
+++ b/Engine/Source/Runtime/LevelSequence/Private/LevelSequencePlayer.cpp
@@ -65,6 +65,8 @@ ULevelSequencePlayer* ULevelSequencePlayer::CreateLevelSequencePlayer(UObject* W
 	ALevelSequenceActor* Actor = World->SpawnActor<ALevelSequenceActor>(SpawnParams);
 
 	Actor->PlaybackSettings = Settings;
+	Actor->SequencePlayer->SetPlaybackSettings(Settings);
+
 	Actor->SetSequence(InLevelSequence);
 
 	Actor->InitializePlayer();
@@ -79,7 +81,7 @@ ULevelSequencePlayer* ULevelSequencePlayer::CreateLevelSequencePlayer(UObject* W
 /* ULevelSequencePlayer implementation
  *****************************************************************************/
 
-void ULevelSequencePlayer::Initialize(ULevelSequence* InLevelSequence, ULevel* InLevel, const FMovieSceneSequencePlaybackSettings& Settings, const FLevelSequenceCameraSettings& InCameraSettings)
+void ULevelSequencePlayer::Initialize(ULevelSequence* InLevelSequence, ULevel* InLevel, const FLevelSequenceCameraSettings& InCameraSettings)
 {
 	// Never use the level to resolve bindings unless we're playing back within a streamed or instanced level
 	StreamedLevelAssetPath = FTopLevelAssetPath();
@@ -103,7 +105,7 @@ void ULevelSequencePlayer::Initialize(ULevelSequence* InLevelSequence, ULevel* I
 	}
 
 	SpawnRegister = MakeShareable(new FLevelSequenceSpawnRegister);
-	UMovieSceneSequencePlayer::Initialize(InLevelSequence, Settings);
+	UMovieSceneSequencePlayer::Initialize(InLevelSequence);
 }
 
 void ULevelSequencePlayer::ResolveBoundObjects(const FGuid& InBindingId, FMovieSceneSequenceID SequenceID, UMovieSceneSequence& InSequence, UObject* ResolutionContext, TArray<UObject*, TInlineAllocator<1>>& OutObjects) const
diff --git a/Engine/Source/Runtime/LevelSequence/Public/LevelSequencePlayer.h b/Engine/Source/Runtime/LevelSequence/Public/LevelSequencePlayer.h
--- a/Engine/Source/Runtime/LevelSequence/Public/LevelSequencePlayer.h
+++ b/Engine/Source/Runtime/LevelSequence/Public/LevelSequencePlayer.h
@@ -101,12 +101,16 @@ public:
 	 *
 	 * @param InLevelSequence The level sequence to play.
 	 * @param InLevel The level that the animation is played in.
-	 * @param Settings The desired playback settings
+	 * @param InCameraSettings The desired camera settings
 	 */
-	void Initialize(ULevelSequence* InLevelSequence, ULevel* InLevel, const FMovieSceneSequencePlaybackSettings& Settings, const FLevelSequenceCameraSettings& InCameraSettings);
-
-	UE_DEPRECATED(4.23, "Added camera settings to Initialize.")
-	void Initialize(ULevelSequence* InLevelSequence, ULevel* InLevel, const FMovieSceneSequencePlaybackSettings& Settings) { Initialize(InLevelSequence, InLevel, Settings, FLevelSequenceCameraSettings()); }
+	void Initialize(ULevelSequence* InLevelSequence, ULevel* InLevel, const FLevelSequenceCameraSettings& InCameraSettings);
+
+	UE_DEPRECATED(5.1, "Use SetPlaybackSettings(...) then Initialize(ULevelSequence*, ULevel*, const FLevelSequenceCameraSettings&)")
+	void Initialize(ULevelSequence* InLevelSequence, ULevel* InLevel, const FMovieSceneSequencePlaybackSettings& Settings, const FLevelSequenceCameraSettings& InCameraSettings)
+	{
+		SetPlaybackSettings(Settings);
+		Initialize(InLevelSequence, InLevel, InCameraSettings);
+	}
 
 public:
 
diff --git a/Engine/Source/Runtime/MovieScene/Private/MovieSceneSequencePlayer.cpp b/Engine/Source/Runtime/MovieScene/Private/MovieSceneSequencePlayer.cpp
--- a/Engine/Source/Runtime/MovieScene/Private/MovieSceneSequencePlayer.cpp
+++ b/Engine/Source/Runtime/MovieScene/Private/MovieSceneSequencePlayer.cpp
@@ -755,7 +755,19 @@ void UMovieSceneSequencePlayer::InitializeForTick(UObject* Context)
 	}
 }
 
+
+void UMovieSceneSequencePlayer::SetPlaybackSettings(const FMovieSceneSequencePlaybackSettings& InSettings)
+{
+	PlaybackSettings = InSettings;
+}
+
 void UMovieSceneSequencePlayer::Initialize(UMovieSceneSequence* InSequence, const FMovieSceneSequencePlaybackSettings& InSettings)
+{
+	PlaybackSettings = InSettings;
+	Initialize(InSequence);
+}
+
+void UMovieSceneSequencePlayer::Initialize(UMovieSceneSequence* InSequence)
 {
 	check(InSequence);
 	check(!bIsEvaluating);
@@ -771,7 +783,6 @@ void UMovieSceneSequencePlayer::Initialize(UMovieSceneSequence* InSequence, cons
 	}
 
 	Sequence = InSequence;
-	PlaybackSettings = InSettings;
 
 	FFrameTime StartTimeWithOffset = StartTime;
 
diff --git a/Engine/Source/Runtime/MovieScene/Public/MovieSceneSequencePlayer.h b/Engine/Source/Runtime/MovieScene/Public/MovieSceneSequencePlayer.h
--- a/Engine/Source/Runtime/MovieScene/Public/MovieSceneSequencePlayer.h
+++ b/Engine/Source/Runtime/MovieScene/Public/MovieSceneSequencePlayer.h
@@ -250,7 +250,7 @@ struct FMovieSceneSequencePlayToParams
 
 template<> struct TStructOpsTypeTraits<FMovieSceneSequencePlaybackSettings> : public TStructOpsTypeTraitsBase2<FMovieSceneSequencePlaybackSettings>
 {
-	enum { WithCopy = true, WithStructuredSerializeFromMismatchedTag = true };
+	enum { WithStructuredSerializeFromMismatchedTag = true };
 };
 
 /**
@@ -510,6 +510,12 @@ public:
 	/** Ensure that this player's tick manager is set up correctly for the specified context */
 	void InitializeForTick(UObject* Context);
 
+	/** Assign this player's playback settings */
+	void SetPlaybackSettings(const FMovieSceneSequencePlaybackSettings& InSettings);
+
+	/** Initialize this player using its existing playback settings */
+	void Initialize(UMovieSceneSequence* InSequence);
+
 	/** Initialize this player with a sequence and some settings */
 	void Initialize(UMovieSceneSequence* InSequence, const FMovieSceneSequencePlaybackSettings& InSettings);
 
