Sequencer: Remove Master from Sequencer

diff --git a/Engine/Plugins/MovieScene/MovieRenderPipeline/Source/MovieRenderPipelineCore/Public/MoviePipelineHighResSetting.h b/Engine/Plugins/MovieScene/MovieRenderPipeline/Source/MovieRenderPipelineCore/Public/MoviePipelineHighResSetting.h
--- a/Engine/Plugins/MovieScene/MovieRenderPipeline/Source/MovieRenderPipelineCore/Public/MoviePipelineHighResSetting.h
+++ b/Engine/Plugins/MovieScene/MovieRenderPipeline/Source/MovieRenderPipelineCore/Public/MoviePipelineHighResSetting.h
@@ -37,7 +37,7 @@ public:
 	virtual FText GetDisplayText() const override { return NSLOCTEXT("MovieRenderPipeline", "HighResSettingDisplayName", "High Resolution"); }
 #endif
 	virtual bool IsValidOnShots() const override { return true; }
-	virtual bool IsValidOnMaster() const override { return true; }
+	virtual bool IsValidOnPrimary() const override { return true; }
 	virtual void ValidateStateImpl() override
 	{
 		Super::ValidateStateImpl();
@@ -82,7 +82,7 @@ public:
 				UMoviePipelineAntiAliasingSetting* AASettings = OwningConfig->FindSetting<UMoviePipelineAntiAliasingSetting>();
 
 				// If we don't have an output setting we know this is a shot override. If this is a shot override we'll fall back
-				// to the cached master configuration values. We don't do this for non-shot overrides as it would pull the cached
+				// to the cached primary configuration values. We don't do this for non-shot overrides as it would pull the cached
 				// values, not the transient ones in the UI.
 				if (!OutputSettings)
 				{
@@ -91,14 +91,14 @@ public:
 						return FText();
 					}
 
-					UMoviePipelineMasterConfig* MasterConfig = InJob->GetConfiguration();
+					UMoviePipelinePrimaryConfig* PrimaryConfig = InJob->GetConfiguration();
 					if (!OutputSettings)
 					{
-						OutputSettings = MasterConfig->FindSetting<UMoviePipelineOutputSetting>();
+						OutputSettings = PrimaryConfig->FindSetting<UMoviePipelineOutputSetting>();
 					}
 					if (!AASettings)
 					{
-						AASettings = MasterConfig->FindSetting<UMoviePipelineAntiAliasingSetting>();
+						AASettings = PrimaryConfig->FindSetting<UMoviePipelineAntiAliasingSetting>();
 					}
 				}
 
diff --git a/Engine/Plugins/MovieScene/MovieRenderPipeline/Source/MovieRenderPipelineCore/Public/MoviePipelinePrimaryConfig.h b/Engine/Plugins/MovieScene/MovieRenderPipeline/Source/MovieRenderPipelineCore/Public/MoviePipelinePrimaryConfig.h
new file mode 100644
--- /dev/null
+++ b/Engine/Plugins/MovieScene/MovieRenderPipeline/Source/MovieRenderPipelineCore/Public/MoviePipelinePrimaryConfig.h
@@ -0,0 +1,94 @@
+// Copyright Epic Games, Inc. All Rights Reserved.
+#pragma once
+
+#include "MoviePipelineConfigBase.h"
+#include "Misc/FrameRate.h"
+
+#include "MoviePipelinePrimaryConfig.generated.h"
+
+
+// Forward Declares
+class ULevelSequence;
+class UMoviePipelineSetting;
+class UMoviePipelineRenderPass;
+class UMoviePipelineOutputBase;
+class UMoviePipelineShotConfig;
+class UMoviePipelineOutputSetting;
+class UMoviePipelineExecutorJob;
+
+/**
+* This class describes the main configuration for a Movie Render Pipeline.
+* Only settings that apply to the entire output should be stored here,
+* anything that is changed on a per-shot basis should be stored inside of 
+* UMovieRenderShotConfig instead.
+*
+* THIS CLASS SHOULD BE IMMUTABLE ONCE PASSED TO THE PIPELINE FOR PROCESSING.
+* (Otherwise you will be modifying the instance that exists in the UI)
+*/
+UCLASS(Blueprintable)
+class MOVIERENDERPIPELINECORE_API UMoviePipelinePrimaryConfig : public UMoviePipelineConfigBase
+{
+	GENERATED_BODY()
+	
+public:
+	UMoviePipelinePrimaryConfig();
+
+public:
+	TArray<UMoviePipelineOutputBase*> GetOutputContainers() const;
+	
+	// UMoviePipelineConfigBase interface
+	virtual TArray<UMoviePipelineSetting*> GetUserSettings() const override;
+	virtual void CopyFrom(UMoviePipelineConfigBase* InConfig) override;
+	// ~UMoviePipelineConfigBase interface
+
+
+	/** Initializes a single instance of every setting so that even non-user-configured settings have a chance to apply their default values. Does nothing if they're already instanced for this configuration. */
+	UFUNCTION(BlueprintCallable, Category = "Movie Render Pipeline")
+	void InitializeTransientSettings();
+
+	UFUNCTION(BlueprintCallable, Category = "Movie Render Pipeline")
+	TArray<UMoviePipelineSetting*> GetTransientSettings() const { return TransientSettings; }
+
+	UFUNCTION(BlueprintCallable, Category = "Movie Render Pipeline")
+	TArray<UMoviePipelineSetting*> GetAllSettings(const bool bIncludeDisabledSettings = false, const bool bIncludeTransientSettings = false) const;
+public:
+
+	/** Returns a pointer to the config specified for the shot, otherwise the default for this pipeline. */
+	UMoviePipelineShotConfig* GetConfigForShot(const FString& ShotName) const;
+
+	void GetFormatArguments(FMoviePipelineFormatArgs& InOutFormatArgs, const bool bIncludeAllSettings = false) const;
+
+
+	/**
+	* Returns the frame rate override from the Primary Configuration (if any) or the Sequence frame rate if no override is specified.
+	* This should be treated as the actual output framerate of the overall Pipeline.
+	*/
+	UFUNCTION(BlueprintCallable, Category = "Movie Render Pipeline")
+	FFrameRate GetEffectiveFrameRate(const ULevelSequence* InSequence) const;
+
+	TRange<FFrameNumber> GetEffectivePlaybackRange(const ULevelSequence* InSequence) const;
+
+protected:
+	virtual bool CanSettingBeAdded(const UMoviePipelineSetting* InSetting) const override
+	{
+		check(InSetting);
+		return InSetting->IsValidOnPrimary();
+	}
+
+	virtual void OnSettingAdded(UMoviePipelineSetting* InSetting) override;
+	virtual void OnSettingRemoved(UMoviePipelineSetting* InSetting) override;
+
+	void AddTransientSettingByClass(const UClass* InSettingClass);
+public:	
+	/** A mapping of Shot Name -> Shot Config to use for rendering specific shots with specific configs. */
+	UPROPERTY(Instanced)
+	TMap<FString, TObjectPtr<UMoviePipelineShotConfig>> PerShotConfigMapping;
+
+private:
+	UPROPERTY(Instanced)
+	TObjectPtr<UMoviePipelineOutputSetting> OutputSetting;
+
+	/** An array of settings that are available in the engine and have not been edited by the user. */
+	UPROPERTY(Transient)
+	TArray<TObjectPtr<UMoviePipelineSetting>> TransientSettings;
+};
\ No newline at end of file
diff --git a/Engine/Plugins/MovieScene/MovieRenderPipeline/Source/MovieRenderPipelineCore/Public/MoviePipelineQueue.h b/Engine/Plugins/MovieScene/MovieRenderPipeline/Source/MovieRenderPipelineCore/Public/MoviePipelineQueue.h
--- a/Engine/Plugins/MovieScene/MovieRenderPipeline/Source/MovieRenderPipelineCore/Public/MoviePipelineQueue.h
+++ b/Engine/Plugins/MovieScene/MovieRenderPipeline/Source/MovieRenderPipelineCore/Public/MoviePipelineQueue.h
@@ -4,14 +4,14 @@
 #include "UObject/Object.h"
 #include "MovieRenderPipelineDataTypes.h"
 #include "LevelSequence.h"
-#include "MoviePipelineMasterConfig.h"
+#include "MoviePipelinePrimaryConfig.h"
 #include "MoviePipelineShotConfig.h"
 #include "MoviePipelineConfigBase.h"
 #include "MovieSceneSequenceID.h"
 
 #include "MoviePipelineQueue.generated.h"
 
-class UMoviePipelineMasterConfig;
+class UMoviePipelinePrimaryConfig;
 class ULevel;
 class ULevelSequence;
 
@@ -203,7 +203,7 @@ public:
 		bEnabled = true;
 		StatusProgress = 0.f;
 		bIsConsumed = false;
-		Configuration = CreateDefaultSubobject<UMoviePipelineMasterConfig>("DefaultConfig");
+		Configuration = CreateDefaultSubobject<UMoviePipelinePrimaryConfig>("DefaultConfig");
 	}
 
 public:	
@@ -336,22 +336,22 @@ public:
 	void OnDuplicated();
 
 	UFUNCTION(BlueprintCallable, Category = "Movie Render Pipeline")
-	void SetPresetOrigin(UMoviePipelineMasterConfig* InPreset);
+	void SetPresetOrigin(UMoviePipelinePrimaryConfig* InPreset);
 
 	UFUNCTION(BlueprintPure, Category = "Movie Render Pipeline")
-	UMoviePipelineMasterConfig* GetPresetOrigin() const 
+	UMoviePipelinePrimaryConfig* GetPresetOrigin() const 
 	{
 		return PresetOrigin.LoadSynchronous();
 	}
 
 	UFUNCTION(BlueprintPure, Category = "Movie Render Pipeline")
-	UMoviePipelineMasterConfig* GetConfiguration() const
+	UMoviePipelinePrimaryConfig* GetConfiguration() const
 	{
 		return Configuration;
 	}
 
 	UFUNCTION(BlueprintCallable, Category = "Movie Render Pipeline")
-	void SetConfiguration(UMoviePipelineMasterConfig* InPreset);
+	void SetConfiguration(UMoviePipelinePrimaryConfig* InPreset);
 
 	UFUNCTION(BlueprintSetter, Category = "Movie Render Pipeline")
 	void SetSequence(FSoftObjectPath InSequence);
@@ -419,12 +419,12 @@ private:
 	/** 
 	*/
 	UPROPERTY(Instanced)
-	TObjectPtr<UMoviePipelineMasterConfig> Configuration;
+	TObjectPtr<UMoviePipelinePrimaryConfig> Configuration;
 
 	/**
 	*/
 	UPROPERTY()
-	TSoftObjectPtr<UMoviePipelineMasterConfig> PresetOrigin;
+	TSoftObjectPtr<UMoviePipelinePrimaryConfig> PresetOrigin;
 
 	/** Whether this job is enabled and should be rendered. */
 	UPROPERTY()
diff --git a/Engine/Plugins/MovieScene/MovieRenderPipeline/Source/MovieRenderPipelineCore/Public/MoviePipelineSetting.h b/Engine/Plugins/MovieScene/MovieRenderPipeline/Source/MovieRenderPipelineCore/Public/MoviePipelineSetting.h
--- a/Engine/Plugins/MovieScene/MovieRenderPipeline/Source/MovieRenderPipelineCore/Public/MoviePipelineSetting.h
+++ b/Engine/Plugins/MovieScene/MovieRenderPipeline/Source/MovieRenderPipelineCore/Public/MoviePipelineSetting.h
@@ -43,13 +43,13 @@ public:
 	void OnMoviePipelineShutdown(UMoviePipeline* InPipeline) { TeardownForPipelineImpl(InPipeline); }
 
 	/**
-	* Called only on settings that have been added to the Master Configuration to let you know that
+	* Called only on settings that have been added to the Primary Configuration to let you know that
 	* a shot is about to be rendered. Useful if your setting needs to know something about the shot
 	* to do something correctly, without using a per-shot override.
 	*/
 	void OnSetupForShot(UMoviePipelineExecutorShot* InShot) { OnSetupForShotImpl(InShot); }
 	/**
-	* Called only on settings that have been added to the Master Configuration to let you know that
+	* Called only on settings that have been added to the Primary Configuration to let you know that
 	* a shot has finished rendering and is being torn down.
 	*/
 	void OnTeardownForShot(UMoviePipelineExecutorShot* InShot) { OnTeardownForShotImpl(InShot); }
@@ -117,8 +117,12 @@ public:
 #endif
 	/** Can this configuration setting be added to shots? If not, it will throw an error when trying to add it to a shot config. */
 	virtual bool IsValidOnShots() const PURE_VIRTUAL(UMoviePipelineSetting::IsValidOnShots, return false; );
-	/** Can this configuration setting be added to the master configuration? If not, it will throw an error when trying to add it to the master configuration. */
-	virtual bool IsValidOnMaster() const PURE_VIRTUAL(UMoviePipelineSetting::IsValidOnMaster, return false; );
+	/** Can this configuration setting be added to the primary configuration? If not, it will throw an error when trying to add it to the primary configuration. */
+	virtual bool IsValidOnPrimary() const PURE_VIRTUAL(UMoviePipelineSetting::IsValidOnPrimary, return false; );
+
+	UE_DEPRECATED(5.2, "IsValidOnMaster is deprecated. Please use IsValidOnPrimary instead")
+	virtual bool IsValidOnMaster() const { return IsValidOnPrimary(); }
+
 	/**
 	* If true, then this setting will be included when searching for settings even if it was added transiently. This is used for the rare case where a setting
 	* needs to be run (to set reasonable default values) even if the user hasn't added it.
diff --git a/Engine/Plugins/MovieScene/MovieRenderPipeline/Source/MovieRenderPipelineCore/Public/MoviePipelineSettingBlueprintBase.h b/Engine/Plugins/MovieScene/MovieRenderPipeline/Source/MovieRenderPipelineCore/Public/MoviePipelineSettingBlueprintBase.h
--- a/Engine/Plugins/MovieScene/MovieRenderPipeline/Source/MovieRenderPipelineCore/Public/MoviePipelineSettingBlueprintBase.h
+++ b/Engine/Plugins/MovieScene/MovieRenderPipeline/Source/MovieRenderPipelineCore/Public/MoviePipelineSettingBlueprintBase.h
@@ -20,7 +20,8 @@ public:
 	UMoviePipelineSetting_BlueprintBase()
 	{
 		CategoryText = NSLOCTEXT("MovieRenderPipelineBP", "DefaultCategoryName_Text", "Custom Settings");
-		bIsValidOnMaster = true;
+		bIsValidOnPrimary = true;
+		bIsValidOnMaster_DEPRECATED = true;
 		bIsValidOnShots = true;
 		bCanBeDisabled = true;
 	}
@@ -68,19 +69,28 @@ public:
 		
 	virtual bool CanBeDisabled() const override { return bCanBeDisabled; }
 #endif
-	virtual bool IsValidOnMaster() const override { return bIsValidOnMaster; }
+	virtual bool IsValidOnPrimary() const override { return bIsValidOnPrimary; }
 	virtual bool IsValidOnShots() const override { return bIsValidOnShots; }
 
+	virtual void PostLoad() override;
+
 protected:
 	UPROPERTY(EditDefaultsOnly, Category = "Movie Pipeline")
 	FText CategoryText;
 	
 	UPROPERTY(EditDefaultsOnly, Category = "Movie Pipeline")
-	bool bIsValidOnMaster;
+	bool bIsValidOnPrimary;
 	
 	UPROPERTY(EditDefaultsOnly, Category = "Movie Pipeline")
 	bool bIsValidOnShots;
 	
 	UPROPERTY(EditDefaultsOnly, Category = "Movie Pipeline")
 	bool bCanBeDisabled;
+
+private:
+
+#if WITH_EDITORONLY_DATA
+	UPROPERTY()
+	bool bIsValidOnMaster_DEPRECATED;
+#endif
 };
\ No newline at end of file
diff --git a/Engine/Plugins/MovieScene/MovieRenderPipeline/Source/MovieRenderPipelineCore/Public/MoviePipelineUtils.h b/Engine/Plugins/MovieScene/MovieRenderPipeline/Source/MovieRenderPipelineCore/Public/MoviePipelineUtils.h
--- a/Engine/Plugins/MovieScene/MovieRenderPipeline/Source/MovieRenderPipelineCore/Public/MoviePipelineUtils.h
+++ b/Engine/Plugins/MovieScene/MovieRenderPipeline/Source/MovieRenderPipelineCore/Public/MoviePipelineUtils.h
@@ -136,7 +136,7 @@ namespace UE
 		MOVIERENDERPIPELINECORE_API void GetHardwareUsageMetadata(TMap<FString, FString>& InFileMetadata, const FString& InOutputDir);
 		MOVIERENDERPIPELINECORE_API void GetMetadataFromCineCamera(class UCineCameraComponent* InComponent, const FString& InCameraName, const FString& InRenderPassName, TMap<FString, FString>& InOutMetadata);
 		MOVIERENDERPIPELINECORE_API void GetMetadataFromCameraLocRot(const FString& InCameraName, const FString& InRenderPassName, const FVector& InCurLoc, const FRotator& InCurRot, const FVector& InPrevLoc, const FRotator& InPrevRot, TMap<FString, FString>& InOutMetadata);
-		MOVIERENDERPIPELINECORE_API FMoviePipelineRenderPassMetrics GetRenderPassMetrics(UMoviePipelineMasterConfig* InMasterConfig, UMoviePipelineExecutorShot* InPipelineExecutorShot, const FMoviePipelineRenderPassMetrics& InRenderPassMetrics, const FIntPoint& InEffectiveOutputResolution);
+		MOVIERENDERPIPELINECORE_API FMoviePipelineRenderPassMetrics GetRenderPassMetrics(UMoviePipelinePrimaryConfig* InPrimaryConfig, UMoviePipelineExecutorShot* InPipelineExecutorShot, const FMoviePipelineRenderPassMetrics& InRenderPassMetrics, const FIntPoint& InEffectiveOutputResolution);
 	}
 }
 
@@ -157,5 +157,5 @@ namespace MoviePipeline
 	/** Given a leaf node, appropriately sets the IsActive flags for the whole hierarchy chain up to root for soloing a shot. */
 	void SetSubSectionHierarchyActive(TSharedPtr<FCameraCutSubSectionHierarchyNode> InRoot, bool bInActive);
 	/** Given a leaf node, searches for sections that will be partially evaluated when using temporal sub-sampling and prints a warning. */
-	void CheckPartialSectionEvaluationAndWarn(const FFrameNumber& LeftDeltaTicks, TSharedPtr<FCameraCutSubSectionHierarchyNode> Node, UMoviePipelineExecutorShot* InShot, const FFrameRate& InMasterDisplayRate);
+	void CheckPartialSectionEvaluationAndWarn(const FFrameNumber& LeftDeltaTicks, TSharedPtr<FCameraCutSubSectionHierarchyNode> Node, UMoviePipelineExecutorShot* InShot, const FFrameRate& InRootDisplayRate);
 }
\ No newline at end of file
diff --git a/Engine/Plugins/MovieScene/MovieRenderPipeline/Source/MovieRenderPipelineCore/Public/MovieRenderPipelineDataTypes.cpp b/Engine/Plugins/MovieScene/MovieRenderPipeline/Source/MovieRenderPipelineCore/Public/MovieRenderPipelineDataTypes.cpp
--- a/Engine/Plugins/MovieScene/MovieRenderPipeline/Source/MovieRenderPipelineCore/Public/MovieRenderPipelineDataTypes.cpp
+++ b/Engine/Plugins/MovieScene/MovieRenderPipeline/Source/MovieRenderPipelineCore/Public/MovieRenderPipelineDataTypes.cpp
@@ -13,7 +13,7 @@ FFrameNumber FMoviePipelineCameraCutInfo::GetOutputFrameCountEstimate() const
 	FFrameNumber OneFrameInTicks = FFrameRate::TransformTime(FFrameTime(FFrameNumber(1)), CachedFrameRate, CachedTickResolution).FloorToFrame();
 
 	// Find out how many ticks long our total output range is.
-	FFrameNumber TotalOutputRangeTicks = TotalOutputRangeMaster.Size<FFrameNumber>();
+	FFrameNumber TotalOutputRangeTicks = TotalOutputRangeRoot.Size<FFrameNumber>();
 	int32 NumFrames = FMath::CeilToInt(TotalOutputRangeTicks.Value / (double)OneFrameInTicks.Value);
 
 	return FFrameNumber(NumFrames);
diff --git a/Engine/Plugins/MovieScene/MovieRenderPipeline/Source/MovieRenderPipelineCore/Public/MovieRenderPipelineDataTypes.h b/Engine/Plugins/MovieScene/MovieRenderPipeline/Source/MovieRenderPipelineCore/Public/MovieRenderPipelineDataTypes.h
--- a/Engine/Plugins/MovieScene/MovieRenderPipeline/Source/MovieRenderPipelineCore/Public/MovieRenderPipelineDataTypes.h
+++ b/Engine/Plugins/MovieScene/MovieRenderPipeline/Source/MovieRenderPipelineCore/Public/MovieRenderPipelineDataTypes.h
@@ -23,7 +23,7 @@ template<class T, class TWeakObjectPtrBase> struct TWeakObjectPtr;
 class UWorld;
 class ULevelSequence;
 class UMovieSceneCinematicShotSection;
-class UMoviePipelineMasterConfig;
+class UMoviePipelinePrimaryConfig;
 class UMoviePipelineShotConfig;
 class UMovieSceneCameraCutSection;
 class UMoviePipelineRenderPass;
@@ -258,7 +258,7 @@ namespace MoviePipeline
 
 	struct FFrameConstantMetrics
 	{
-		/** What is the tick resolution of the master sequence */
+		/** What is the tick resolution of the root sequence */
 		FFrameRate TickResolution;
 		/** What is the tick resolution of the current shot */
 		FFrameRate ShotTickResolution;
@@ -520,23 +520,23 @@ public:
 	/** Cached Tick Resolution our numbers are in. Simplifies some APIs. */
 	FFrameRate CachedTickResolution;
 	
-	/** Cached Tick Resolution the movie scene this range was generated for is in. Can be different than the master due to mixed tick resolutions. */
+	/** Cached Tick Resolution the movie scene this range was generated for is in. Can be different than the root due to mixed tick resolutions. */
 	FFrameRate CachedShotTickResolution;
 public:
 	/** The current state of processing this Shot is in. Not all states will be passed through. */
 	EMovieRenderShotState State;
 
-	/** The current tick of this shot that we're on in master space */
-	FFrameNumber CurrentTickInMaster;
+	/** The current tick of this shot that we're on in root space */
+	FFrameNumber CurrentTickInRoot;
 
 	/** Converts from the outermost space into the innermost space. Only works with linear transforms. */
 	FMovieSceneTimeTransform OuterToInnerTransform;
 	
-	/** The total range of output frames in master space */
-	TRange<FFrameNumber> TotalOutputRangeMaster;
+	/** The total range of output frames in root space */
+	TRange<FFrameNumber> TotalOutputRangeRoot;
 	
-	/** The total range of warmup frames in master space */
-	TRange<FFrameNumber> WarmupRangeMaster;
+	/** The total range of warmup frames in root space */
+	TRange<FFrameNumber> WarmupRangeRoot;
 
 	/** Metrics - How much work has been done for this particular shot and how much do we estimate we have to do? */
 	FMoviePipelineSegmentWorkMetrics WorkMetrics;
@@ -801,7 +801,7 @@ struct FMoviePipelineFilenameResolveParams
 	{
 	}
 	
-	/** Frame Number for the Master (matching what you see in the Sequencer timeline. ie: If the Sequence PlaybackRange starts on 50, this value would be 50 on the first frame.*/
+	/** Frame Number for the Root (matching what you see in the Sequencer timeline. ie: If the Sequence PlaybackRange starts on 50, this value would be 50 on the first frame.*/
 	UPROPERTY(BlueprintReadWrite, EditAnywhere, Category = "Movie Render Pipeline")
 	int32 FrameNumber;
 	
@@ -809,7 +809,7 @@ struct FMoviePipelineFilenameResolveParams
 	UPROPERTY(BlueprintReadWrite, EditAnywhere, Category = "Movie Render Pipeline")
 	int32 FrameNumberShot;
 	
-	/** Frame Number for the Master (relative to 0, not what you would see in the Sequencer timeline. ie: If sequence PlaybackRange starts on 50, this value would be 0 on the first frame. */
+	/** Frame Number for the Root (relative to 0, not what you would see in the Sequencer timeline. ie: If sequence PlaybackRange starts on 50, this value would be 0 on the first frame. */
 	UPROPERTY(BlueprintReadWrite, EditAnywhere, Category = "Movie Render Pipeline")
 	int32 FrameNumberRel;
 	
@@ -863,7 +863,7 @@ struct FMoviePipelineFilenameResolveParams
 	/** If this shot has multiple cameras in the sidecar array, which camera is this for? -1 will return the InnerName/Main camera for the shot. */
 	int32 CameraIndex;
 
-	/** Optional. If specified, settings will be pulled from this shot (if overriden by the shot). If null, always use the master configuration in the job. */
+	/** Optional. If specified, settings will be pulled from this shot (if overriden by the shot). If null, always use the root configuration in the job. */
 	UPROPERTY(BlueprintReadWrite, EditAnywhere, Category = "Movie Render Pipeline")
 	TObjectPtr<UMoviePipelineExecutorShot> ShotOverride;
 
diff --git a/Engine/Plugins/MovieScene/MovieRenderPipeline/Source/MovieRenderPipelineEditor/Private/AssetTypeActions_PipelineConfigs.cpp b/Engine/Plugins/MovieScene/MovieRenderPipeline/Source/MovieRenderPipelineEditor/Private/AssetTypeActions_PipelineConfigs.cpp
--- a/Engine/Plugins/MovieScene/MovieRenderPipeline/Source/MovieRenderPipelineEditor/Private/AssetTypeActions_PipelineConfigs.cpp
+++ b/Engine/Plugins/MovieScene/MovieRenderPipeline/Source/MovieRenderPipelineEditor/Private/AssetTypeActions_PipelineConfigs.cpp
@@ -7,7 +7,7 @@
 #include "IMovieRenderPipelineEditorModule.h"
 #include "Widgets/Notifications/SNotificationList.h"
 
-void FAssetTypeActions_PipelineMasterConfig::OpenAssetEditor( const TArray<UObject*>& InObjects, TSharedPtr<IToolkitHost> EditWithinLevelEditor )
+void FAssetTypeActions_PipelinePrimaryConfig::OpenAssetEditor( const TArray<UObject*>& InObjects, TSharedPtr<IToolkitHost> EditWithinLevelEditor )
 {
 	UAssetEditorSubsystem* AssetEditorSubsystem = GEditor->GetEditorSubsystem<UAssetEditorSubsystem>();
 	if (AssetEditorSubsystem)
diff --git a/Engine/Plugins/MovieScene/MovieRenderPipeline/Source/MovieRenderPipelineEditor/Private/AssetTypeActions_PipelineConfigs.h b/Engine/Plugins/MovieScene/MovieRenderPipeline/Source/MovieRenderPipelineEditor/Private/AssetTypeActions_PipelineConfigs.h
--- a/Engine/Plugins/MovieScene/MovieRenderPipeline/Source/MovieRenderPipelineEditor/Private/AssetTypeActions_PipelineConfigs.h
+++ b/Engine/Plugins/MovieScene/MovieRenderPipeline/Source/MovieRenderPipelineEditor/Private/AssetTypeActions_PipelineConfigs.h
@@ -5,17 +5,17 @@
 #include "CoreMinimal.h"
 #include "Toolkits/IToolkitHost.h"
 #include "AssetTypeActions_Base.h"
-#include "MoviePipelineMasterConfig.h"
+#include "MoviePipelinePrimaryConfig.h"
 #include "MoviePipelineShotConfig.h"
 #include "MoviePipelineQueue.h"
 
-class FAssetTypeActions_PipelineMasterConfig : public FAssetTypeActions_Base
+class FAssetTypeActions_PipelinePrimaryConfig : public FAssetTypeActions_Base
 {
 public:
 	// IAssetTypeActions Implementation
-	virtual FText GetName() const override { return NSLOCTEXT("AssetTypeActions", "AssetTypeActions_PipelineMasterConfig", "Movie Pipeline Master Config"); }
+	virtual FText GetName() const override { return NSLOCTEXT("AssetTypeActions", "AssetTypeActions_PipelinePrimaryConfig", "Movie Pipeline Primary Config"); }
 	virtual FColor GetTypeColor() const override { return FColor(78, 40, 165); }
-	virtual UClass* GetSupportedClass() const override { return UMoviePipelineMasterConfig::StaticClass(); }
+	virtual UClass* GetSupportedClass() const override { return UMoviePipelinePrimaryConfig::StaticClass(); }
 	virtual void OpenAssetEditor( const TArray<UObject*>& InObjects, TSharedPtr<class IToolkitHost> EditWithinLevelEditor = TSharedPtr<IToolkitHost>() ) override;
 	virtual bool CanFilter() override { return false; }
 	virtual uint32 GetCategories() override { return EAssetTypeCategories::Animation; }
diff --git a/Engine/Plugins/MovieScene/MovieRenderPipeline/Source/MovieRenderPipelineEditor/Private/MoviePipelineConfigFactory.cpp b/Engine/Plugins/MovieScene/MovieRenderPipeline/Source/MovieRenderPipelineEditor/Private/MoviePipelineConfigFactory.cpp
--- a/Engine/Plugins/MovieScene/MovieRenderPipeline/Source/MovieRenderPipelineEditor/Private/MoviePipelineConfigFactory.cpp
+++ b/Engine/Plugins/MovieScene/MovieRenderPipeline/Source/MovieRenderPipelineEditor/Private/MoviePipelineConfigFactory.cpp
@@ -1,25 +1,25 @@
 // Copyright Epic Games, Inc. All Rights Reserved.
 
 #include "MoviePipelineConfigFactory.h"
-#include "MoviePipelineMasterConfig.h"
+#include "MoviePipelinePrimaryConfig.h"
 #include "MoviePipelineShotConfig.h"
 #include "AssetTypeCategories.h"
 
 #include UE_INLINE_GENERATED_CPP_BY_NAME(MoviePipelineConfigFactory)
 
-UMoviePipelineMasterConfigFactory::UMoviePipelineMasterConfigFactory()
+UMoviePipelinePrimaryConfigFactory::UMoviePipelinePrimaryConfigFactory()
 {
 	bCreateNew = true;
 	bEditAfterNew = false;
-	SupportedClass = UMoviePipelineMasterConfig::StaticClass();
+	SupportedClass = UMoviePipelinePrimaryConfig::StaticClass();
 }
 
-UObject* UMoviePipelineMasterConfigFactory::FactoryCreateNew(UClass* Class, UObject* InParent, FName Name, EObjectFlags Flags, UObject* Context, FFeedbackContext* Warn)
+UObject* UMoviePipelinePrimaryConfigFactory::FactoryCreateNew(UClass* Class, UObject* InParent, FName Name, EObjectFlags Flags, UObject* Context, FFeedbackContext* Warn)
 {
-	return NewObject<UMoviePipelineMasterConfig>(InParent, Class, Name, Flags);
+	return NewObject<UMoviePipelinePrimaryConfig>(InParent, Class, Name, Flags);
 }
 
-uint32 UMoviePipelineMasterConfigFactory::GetMenuCategories() const
+uint32 UMoviePipelinePrimaryConfigFactory::GetMenuCategories() const
 {
 	return EAssetTypeCategories::Type::None; 
 }
diff --git a/Engine/Plugins/MovieScene/MovieRenderPipeline/Source/MovieRenderPipelineEditor/Private/MoviePipelineEditorBlueprintLibrary.cpp b/Engine/Plugins/MovieScene/MovieRenderPipeline/Source/MovieRenderPipelineEditor/Private/MoviePipelineEditorBlueprintLibrary.cpp
--- a/Engine/Plugins/MovieScene/MovieRenderPipeline/Source/MovieRenderPipelineEditor/Private/MoviePipelineEditorBlueprintLibrary.cpp
+++ b/Engine/Plugins/MovieScene/MovieRenderPipeline/Source/MovieRenderPipelineEditor/Private/MoviePipelineEditorBlueprintLibrary.cpp
@@ -25,7 +25,7 @@
 
 #define LOCTEXT_NAMESPACE "MoviePipelineEditorBlueprintLibrary"
 
-bool UMoviePipelineEditorBlueprintLibrary::ExportConfigToAsset(const UMoviePipelineMasterConfig* InConfig, const FString& InPackagePath, const FString& InFileName, const bool bInSaveAsset, UMoviePipelineMasterConfig*& OutAsset, FText& OutErrorReason)
+bool UMoviePipelineEditorBlueprintLibrary::ExportConfigToAsset(const UMoviePipelinePrimaryConfig* InConfig, const FString& InPackagePath, const FString& InFileName, const bool bInSaveAsset, UMoviePipelinePrimaryConfig*& OutAsset, FText& OutErrorReason)
 {
 	if(!InConfig)
 	{
@@ -46,7 +46,7 @@ bool UMoviePipelineEditorBlueprintLibrary::ExportConfigToAsset(const UMoviePipel
 	NewPackage->AddToRoot();
 	
 	// Duplicate the provided config into this package.
-	UMoviePipelineMasterConfig* NewConfig = Cast<UMoviePipelineMasterConfig>(StaticDuplicateObject(InConfig, NewPackage, FName(*InFileName), RF_NoFlags));
+	UMoviePipelinePrimaryConfig* NewConfig = Cast<UMoviePipelinePrimaryConfig>(StaticDuplicateObject(InConfig, NewPackage, FName(*InFileName), RF_NoFlags));
 	NewConfig->SetFlags(RF_Public | RF_Transactional | RF_Standalone);
 	NewConfig->MarkPackageDirty();
 
diff --git a/Engine/Plugins/MovieScene/MovieRenderPipeline/Source/MovieRenderPipelineEditor/Private/MoviePipelineNewProcessExecutor.cpp b/Engine/Plugins/MovieScene/MovieRenderPipeline/Source/MovieRenderPipelineEditor/Private/MoviePipelineNewProcessExecutor.cpp
--- a/Engine/Plugins/MovieScene/MovieRenderPipeline/Source/MovieRenderPipelineEditor/Private/MoviePipelineNewProcessExecutor.cpp
+++ b/Engine/Plugins/MovieScene/MovieRenderPipeline/Source/MovieRenderPipelineEditor/Private/MoviePipelineNewProcessExecutor.cpp
@@ -148,7 +148,7 @@ void UMoviePipelineNewProcessExecutor::Execute_Implementation(UMoviePipelineQueu
 		FString PipelineConfig = TEXT("MovieRenderPipeline/QueueManifest.utxt");
 		
 		// Because the Queue has multiple jobs in it, we don't need to pass which sequence to render. That's only needed if you're rendering a 
-		// specific sequence with a specific master config.
+		// specific sequence with a specific primary config.
 		MoviePipelineArgs = FString::Printf(TEXT("-MoviePipelineConfig=\"%s\""), *PipelineConfig); // -MoviePipeline=\"%s\" -MoviePipelineLocalExecutorClass=\"%s\" -MoviePipelineClass=\"%s\""),
 	}
 
diff --git a/Engine/Plugins/MovieScene/MovieRenderPipeline/Source/MovieRenderPipelineEditor/Private/MovieRenderPipelineEditorModule.cpp b/Engine/Plugins/MovieScene/MovieRenderPipeline/Source/MovieRenderPipelineEditor/Private/MovieRenderPipelineEditorModule.cpp
--- a/Engine/Plugins/MovieScene/MovieRenderPipeline/Source/MovieRenderPipelineEditor/Private/MovieRenderPipelineEditorModule.cpp
+++ b/Engine/Plugins/MovieScene/MovieRenderPipeline/Source/MovieRenderPipelineEditor/Private/MovieRenderPipelineEditorModule.cpp
@@ -201,8 +201,8 @@ void FMovieRenderPipelineEditorModule::StartupModule()
 
 	IAssetTools& AssetTools = FModuleManager::LoadModuleChecked<FAssetToolsModule>("AssetTools").Get();
 	{
-		MasterConfigAssetActions = MakeShared<FAssetTypeActions_PipelineMasterConfig>();
-		AssetTools.RegisterAssetTypeActions(MasterConfigAssetActions.ToSharedRef());
+		PrimaryConfigAssetActions = MakeShared<FAssetTypeActions_PipelinePrimaryConfig>();
+		AssetTools.RegisterAssetTypeActions(PrimaryConfigAssetActions.ToSharedRef());
 
 		ShotConfigAssetActions = MakeShared<FAssetTypeActions_PipelineShotConfig>();
 		AssetTools.RegisterAssetTypeActions(ShotConfigAssetActions.ToSharedRef());
@@ -217,7 +217,7 @@ void FMovieRenderPipelineEditorModule::ShutdownModule()
 	FAssetToolsModule* AssetToolsModule = FModuleManager::GetModulePtr<FAssetToolsModule>("AssetTools");
 	if (AssetToolsModule)
 	{
-		AssetToolsModule->Get().UnregisterAssetTypeActions(MasterConfigAssetActions.ToSharedRef());
+		AssetToolsModule->Get().UnregisterAssetTypeActions(PrimaryConfigAssetActions.ToSharedRef());
 		AssetToolsModule->Get().UnregisterAssetTypeActions(ShotConfigAssetActions.ToSharedRef());
 		AssetToolsModule->Get().UnregisterAssetTypeActions(QueueAssetActions.ToSharedRef());
 	}
diff --git a/Engine/Plugins/MovieScene/MovieRenderPipeline/Source/MovieRenderPipelineEditor/Private/MovieRenderPipelineEditorModule.h b/Engine/Plugins/MovieScene/MovieRenderPipeline/Source/MovieRenderPipelineEditor/Private/MovieRenderPipelineEditorModule.h
--- a/Engine/Plugins/MovieScene/MovieRenderPipeline/Source/MovieRenderPipelineEditor/Private/MovieRenderPipelineEditorModule.h
+++ b/Engine/Plugins/MovieScene/MovieRenderPipeline/Source/MovieRenderPipelineEditor/Private/MovieRenderPipelineEditorModule.h
@@ -7,7 +7,7 @@
 #include "Delegates/IDelegateInstance.h"
 
 class UMoviePipelineShotConfig;
-class UMoviePipelineMasterConfig;
+class UMoviePipelinePrimaryConfig;
 class UMoviePipelineExecutorBase;
 
 class FMovieRenderPipelineEditorModule : public IMovieRenderPipelineEditorModule
@@ -24,7 +24,7 @@ private:
 	void RegisterMovieRenderer();
 	void UnregisterMovieRenderer();
 
-	TSharedPtr<class FAssetTypeActions_PipelineMasterConfig> MasterConfigAssetActions;
+	TSharedPtr<class FAssetTypeActions_PipelinePrimaryConfig> PrimaryConfigAssetActions;
 	TSharedPtr<class FAssetTypeActions_PipelineShotConfig> ShotConfigAssetActions;
 	TSharedPtr<class FAssetTypeActions_PipelineQueue> QueueAssetActions;
 	FDelegateHandle MovieRendererDelegate;
diff --git a/Engine/Plugins/MovieScene/MovieRenderPipeline/Source/MovieRenderPipelineEditor/Private/Widgets/SMoviePipelineConfigPanel.cpp b/Engine/Plugins/MovieScene/MovieRenderPipeline/Source/MovieRenderPipelineEditor/Private/Widgets/SMoviePipelineConfigPanel.cpp
--- a/Engine/Plugins/MovieScene/MovieRenderPipeline/Source/MovieRenderPipelineEditor/Private/Widgets/SMoviePipelineConfigPanel.cpp
+++ b/Engine/Plugins/MovieScene/MovieRenderPipeline/Source/MovieRenderPipelineEditor/Private/Widgets/SMoviePipelineConfigPanel.cpp
@@ -241,7 +241,7 @@ PRAGMA_ENABLE_OPTIMIZATION
 
 FText SMoviePipelineConfigPanel::GetConfigTypeLabel() const
 {
-	return LOCTEXT("MasterConfigType_Label", "MASTER");
+	return LOCTEXT("PrimaryConfigType_Label", "PRIMARY");
 }
 
 FReply SMoviePipelineConfigPanel::OnConfirmChanges()
@@ -368,6 +368,15 @@ TSharedRef<SWidget> SMoviePipelineConfigPanel::OnGeneratePresetsMenu()
 
 		AssetPickerConfig.AssetShowWarningText = LOCTEXT("NoPresets_Warning", "No Presets Found");
 		AssetPickerConfig.Filter.ClassPaths.Add(ConfigAssetType->GetClassPathName());
+
+		// The redirected masterconfig assets aren't found when the new primaryconfig asset path is specified. This is a temporary workaround 
+		// for the AssetRegistry not applying CoreRedirects to the ClassPath. 
+		// This can be removed once that bug is fixed and the KnownCase in CheckUnacceptableWords can also be removed.
+		if (ConfigAssetType->GetClassPathName() == FTopLevelAssetPath(TEXT("/Script/MovieRenderPipelineCore.MoviePipelinePrimaryConfig")))
+		{
+			AssetPickerConfig.Filter.ClassPaths.Add(FTopLevelAssetPath(TEXT("/Script/MovieRenderPipelineCore.MoviePipelineMasterConfig")));
+		}
+		
 		AssetPickerConfig.Filter.bRecursiveClasses = true;
 		AssetPickerConfig.OnAssetSelected = FOnAssetSelected::CreateSP(this, &SMoviePipelineConfigPanel::OnImportPreset);
 	}
diff --git a/Engine/Plugins/MovieScene/MovieRenderPipeline/Source/MovieRenderPipelineEditor/Private/Widgets/SMoviePipelineConfigPanel.h b/Engine/Plugins/MovieScene/MovieRenderPipeline/Source/MovieRenderPipelineEditor/Private/Widgets/SMoviePipelineConfigPanel.h
--- a/Engine/Plugins/MovieScene/MovieRenderPipeline/Source/MovieRenderPipelineEditor/Private/Widgets/SMoviePipelineConfigPanel.h
+++ b/Engine/Plugins/MovieScene/MovieRenderPipeline/Source/MovieRenderPipelineEditor/Private/Widgets/SMoviePipelineConfigPanel.h
@@ -110,7 +110,7 @@ private:
 	/** The main movie pipeline editor widget */
 	TSharedPtr<SMoviePipelineConfigEditor> MoviePipelineEditorWidget;
 
-	/** What type of asset are we editing? Could be a master config or a per-shot override config. */
+	/** What type of asset are we editing? Could be a primary config or a per-shot override config. */
 	TSubclassOf<UMoviePipelineConfigBase> ConfigAssetType;
 
 	FOnMoviePipelineConfigChanged OnConfigurationModified;
diff --git a/Engine/Plugins/MovieScene/MovieRenderPipeline/Source/MovieRenderPipelineEditor/Private/Widgets/SMoviePipelineQueueEditor.cpp b/Engine/Plugins/MovieScene/MovieRenderPipeline/Source/MovieRenderPipelineEditor/Private/Widgets/SMoviePipelineQueueEditor.cpp
--- a/Engine/Plugins/MovieScene/MovieRenderPipeline/Source/MovieRenderPipelineEditor/Private/Widgets/SMoviePipelineQueueEditor.cpp
+++ b/Engine/Plugins/MovieScene/MovieRenderPipeline/Source/MovieRenderPipelineEditor/Private/Widgets/SMoviePipelineQueueEditor.cpp
@@ -7,7 +7,7 @@
 #include "MoviePipelineQueue.h"
 #include "MoviePipelineOutputSetting.h"
 #include "MoviePipelineQueueSubsystem.h"
-#include "MoviePipelineMasterConfig.h"
+#include "MoviePipelinePrimaryConfig.h"
 #include "MovieRenderPipelineStyle.h"
 #include "MovieRenderPipelineSettings.h"
 #include "Sections/MovieSceneCameraCutSection.h"
@@ -198,7 +198,7 @@ public:
 		return FText();
 	}
 
-	FText GetMasterConfigLabel() const
+	FText GetPrimaryConfigLabel() const
 	{
 		UMoviePipelineExecutorJob* Job = WeakJob.Get();
 		if (Job)
@@ -226,7 +226,7 @@ public:
 		UMoviePipelineExecutorJob* Job = WeakJob.Get();
 		if (Job)
 		{
-			Job->SetPresetOrigin(CastChecked<UMoviePipelineMasterConfig>(AssetData.GetAsset()));
+			Job->SetPresetOrigin(CastChecked<UMoviePipelinePrimaryConfig>(AssetData.GetAsset()));
 		}
 
 		OnChosePresetCallback.ExecuteIfBound(WeakJob, nullptr);
@@ -248,7 +248,7 @@ public:
 		OnChosePresetCallback.ExecuteIfBound(WeakJob, nullptr);
 	}
 
-	EVisibility GetMasterConfigModifiedVisibility() const
+	EVisibility GetPrimaryConfigModifiedVisibility() const
 	{
 		UMoviePipelineExecutorJob* Job = WeakJob.Get();
 		if (Job)
@@ -259,7 +259,7 @@ public:
 		return EVisibility::Collapsed;
 	}
 
-	void OnEditMasterConfigForJob()
+	void OnEditPrimaryConfigForJob()
 	{
 		OnEditConfigCallback.ExecuteIfBound(WeakJob, nullptr);
 	}
@@ -339,7 +339,7 @@ public:
 
 	TSharedRef<SWidget> OnGenerateConfigPresetPickerMenu()
 	{
-		return OnGenerateConfigPresetPickerMenuFromClass(UMoviePipelineMasterConfig::StaticClass(),
+		return OnGenerateConfigPresetPickerMenuFromClass(UMoviePipelinePrimaryConfig::StaticClass(),
 			FOnAssetSelected::CreateRaw(this, &FMoviePipelineQueueJobTreeItem::OnPickPresetFromAsset),
 			FExecuteAction::CreateRaw(this, &FMoviePipelineQueueJobTreeItem::OnPickNewPreset)
 			);
@@ -366,7 +366,7 @@ public:
 			AssetPickerConfig.ThumbnailScale = 0.1f;
 			AssetPickerConfig.SaveSettingsName = TEXT("MoviePipelineConfigAsset");
 
-			AssetPickerConfig.AssetShowWarningText = LOCTEXT("NoConfigs_Warning", "No Master Configurations Found");
+			AssetPickerConfig.AssetShowWarningText = LOCTEXT("NoConfigs_Warning", "No Primary Configurations Found");
 			AssetPickerConfig.Filter.ClassPaths.Add(InClass->GetClassPathName());
 			AssetPickerConfig.OnAssetSelected = InOnAssetSelected;
 		}
@@ -461,8 +461,8 @@ TSharedRef<SWidget> SQueueJobListRow::GenerateWidgetForColumn(const FName& Colum
 		.Padding(2, 0)
 		[
 			SNew(SHyperlink)
-			.Text(TAttribute<FText>::Create(TAttribute<FText>::FGetter::CreateSP(Item.Get(), &FMoviePipelineQueueJobTreeItem::GetMasterConfigLabel)))
-			.OnNavigate(Item.Get(), &FMoviePipelineQueueJobTreeItem::OnEditMasterConfigForJob)
+			.Text(TAttribute<FText>::Create(TAttribute<FText>::FGetter::CreateSP(Item.Get(), &FMoviePipelineQueueJobTreeItem::GetPrimaryConfigLabel)))
+			.OnNavigate(Item.Get(), &FMoviePipelineQueueJobTreeItem::OnEditPrimaryConfigForJob)
 		]
 
 		+ SHorizontalBox::Slot()
@@ -470,7 +470,7 @@ TSharedRef<SWidget> SQueueJobListRow::GenerateWidgetForColumn(const FName& Colum
 		[
 			SNew(STextBlock)
 			.Text(LOCTEXT("ModifiedConfigIndicator", "*"))
-			.Visibility(Item.Get(), &FMoviePipelineQueueJobTreeItem::GetMasterConfigModifiedVisibility)
+			.Visibility(Item.Get(), &FMoviePipelineQueueJobTreeItem::GetPrimaryConfigModifiedVisibility)
 		]
 
 		+ SHorizontalBox::Slot()
diff --git a/Engine/Plugins/MovieScene/MovieRenderPipeline/Source/MovieRenderPipelineEditor/Private/Widgets/SMoviePipelineQueuePanel.cpp b/Engine/Plugins/MovieScene/MovieRenderPipeline/Source/MovieRenderPipelineEditor/Private/Widgets/SMoviePipelineQueuePanel.cpp
--- a/Engine/Plugins/MovieScene/MovieRenderPipeline/Source/MovieRenderPipelineEditor/Private/Widgets/SMoviePipelineQueuePanel.cpp
+++ b/Engine/Plugins/MovieScene/MovieRenderPipeline/Source/MovieRenderPipelineEditor/Private/Widgets/SMoviePipelineQueuePanel.cpp
@@ -10,7 +10,7 @@
 #include "MoviePipelineBlueprintLibrary.h"
 #include "Sections/MovieSceneCinematicShotSection.h"
 #include "MoviePipelineQueue.h"
-#include "MoviePipelineMasterConfig.h"
+#include "MoviePipelinePrimaryConfig.h"
 #include "MoviePipelineQueueSubsystem.h"
 
 // Slate Includes
@@ -371,7 +371,7 @@ void SMoviePipelineQueuePanel::OnEditJobConfigRequested(TWeakObjectPtr<UMoviePip
 	}
 	else
 	{
-		ConfigType = UMoviePipelineMasterConfig::StaticClass();
+		ConfigType = UMoviePipelinePrimaryConfig::StaticClass();
 		BasePreset = InJob->GetPresetOrigin();
 		BaseConfig = InJob->GetConfiguration();
 	}
@@ -422,9 +422,9 @@ void SMoviePipelineQueuePanel::OnConfigUpdatedForJob(TWeakObjectPtr<UMoviePipeli
 		}
 		else
 		{
-			if (UMoviePipelineMasterConfig* MasterConfig = Cast<UMoviePipelineMasterConfig>(InConfig))
+			if (UMoviePipelinePrimaryConfig* PrimaryConfig = Cast<UMoviePipelinePrimaryConfig>(InConfig))
 			{
-				InJob->SetConfiguration(MasterConfig);
+				InJob->SetConfiguration(PrimaryConfig);
 			}
 		}
 	}
@@ -445,9 +445,9 @@ void SMoviePipelineQueuePanel::OnConfigUpdatedForJobToPreset(TWeakObjectPtr<UMov
 		}
 		else
 		{
-			if (UMoviePipelineMasterConfig* MasterConfig = Cast<UMoviePipelineMasterConfig>(InConfig))
+			if (UMoviePipelinePrimaryConfig* PrimaryConfig = Cast<UMoviePipelinePrimaryConfig>(InConfig))
 			{
-				InJob->SetPresetOrigin(MasterConfig);
+				InJob->SetPresetOrigin(PrimaryConfig);
 			}
 		}
 	}
diff --git a/Engine/Plugins/MovieScene/MovieRenderPipeline/Source/MovieRenderPipelineEditor/Public/MoviePipelineConfigFactory.h b/Engine/Plugins/MovieScene/MovieRenderPipeline/Source/MovieRenderPipelineEditor/Public/MoviePipelineConfigFactory.h
--- a/Engine/Plugins/MovieScene/MovieRenderPipeline/Source/MovieRenderPipelineEditor/Public/MoviePipelineConfigFactory.h
+++ b/Engine/Plugins/MovieScene/MovieRenderPipeline/Source/MovieRenderPipelineEditor/Public/MoviePipelineConfigFactory.h
@@ -7,11 +7,11 @@
 
 // Forward Declare
 UCLASS(BlueprintType)
-class MOVIERENDERPIPELINEEDITOR_API UMoviePipelineMasterConfigFactory : public UFactory
+class MOVIERENDERPIPELINEEDITOR_API UMoviePipelinePrimaryConfigFactory : public UFactory
 {
     GENERATED_BODY()
 public:
-	UMoviePipelineMasterConfigFactory();
+	UMoviePipelinePrimaryConfigFactory();
 	//~ Begin UFactory Interface
 	virtual UObject* FactoryCreateNew(UClass* Class, UObject* InParent, FName Name, EObjectFlags Flags, UObject* Context, FFeedbackContext* Warn) override;
 	virtual uint32 GetMenuCategories() const override;
diff --git a/Engine/Plugins/MovieScene/MovieRenderPipeline/Source/MovieRenderPipelineEditor/Public/MoviePipelineEditorBlueprintLibrary.h b/Engine/Plugins/MovieScene/MovieRenderPipeline/Source/MovieRenderPipelineEditor/Public/MoviePipelineEditorBlueprintLibrary.h
--- a/Engine/Plugins/MovieScene/MovieRenderPipeline/Source/MovieRenderPipelineEditor/Public/MoviePipelineEditorBlueprintLibrary.h
+++ b/Engine/Plugins/MovieScene/MovieRenderPipeline/Source/MovieRenderPipelineEditor/Public/MoviePipelineEditorBlueprintLibrary.h
@@ -7,12 +7,12 @@
 #include "UObject/TextProperty.h"
 #include "Internationalization/Text.h"
 #include "Containers/UnrealString.h"
-#include "MoviePipelineMasterConfig.h"
+#include "MoviePipelinePrimaryConfig.h"
 
 #include "MoviePipelineEditorBlueprintLibrary.generated.h"
 
 // Forward Declare
-class UMoviePipelineMasterConfig;
+class UMoviePipelinePrimaryConfig;
 class UMoviePipelineExecutorJob;
 class ULevelSequence;
 
@@ -23,7 +23,7 @@ class MOVIERENDERPIPELINEEDITOR_API UMoviePipelineEditorBlueprintLibrary : publi
 	
 public:
 	UFUNCTION(BlueprintCallable, Category = "Movie Render Pipeline")
-	static bool ExportConfigToAsset(const UMoviePipelineMasterConfig* InConfig, const FString& InPackagePath, const FString& InFileName, const bool bInSaveAsset, UMoviePipelineMasterConfig*& OutAsset, FText& OutErrorReason);
+	static bool ExportConfigToAsset(const UMoviePipelinePrimaryConfig* InConfig, const FString& InPackagePath, const FString& InFileName, const bool bInSaveAsset, UMoviePipelinePrimaryConfig*& OutAsset, FText& OutErrorReason);
 
 	/** Checks to see if any of the Jobs try to point to maps that wouldn't be valid on a remote render (ie: unsaved maps) */
 	UFUNCTION(BlueprintPure, Category = "Movie Render Pipeline")
diff --git a/Engine/Plugins/MovieScene/MovieRenderPipeline/Source/MovieRenderPipelineEditor/Public/MovieRenderPipelineSettings.h b/Engine/Plugins/MovieScene/MovieRenderPipeline/Source/MovieRenderPipelineEditor/Public/MovieRenderPipelineSettings.h
--- a/Engine/Plugins/MovieScene/MovieRenderPipeline/Source/MovieRenderPipelineEditor/Public/MovieRenderPipelineSettings.h
+++ b/Engine/Plugins/MovieScene/MovieRenderPipeline/Source/MovieRenderPipelineEditor/Public/MovieRenderPipelineSettings.h
@@ -10,7 +10,7 @@
 
 class UMoviePipelineExecutorBase;
 class UMoviePipeline;
-class UMoviePipelineMasterConfig;
+class UMoviePipelinePrimaryConfig;
 
 /**
  * Universal Movie Render Pipeline settings that apply to the whole project.
@@ -34,7 +34,7 @@ public:
 	* What was the last configuration preset the user used? Can be null.
 	*/
 	UPROPERTY(Transient, BlueprintReadWrite, Category = "Movie Render Pipeline")
-	TSoftObjectPtr<UMoviePipelineMasterConfig> LastPresetOrigin;
+	TSoftObjectPtr<UMoviePipelinePrimaryConfig> LastPresetOrigin;
 	
 	/**
 	* When the user uses the UI to request we render a movie locally, which implementation should we use
@@ -71,7 +71,7 @@ public:
 	FSoftClassPath DefaultPipeline;
 
 	/**
-	* The settings specified here will automatically be added to a Movie Pipeline Master Configuration when using the UI. 
+	* The settings specified here will automatically be added to a Movie Pipeline Primary Configuration when using the UI. 
 	* This does not apply to scripting and does not apply to runtime. It is only a convenience function so that when a job is
 	* created, it can be pre-filled with some settings to make the render functional out of the gate. It can also be
 	* used to automatically add your own setting to jobs.
diff --git a/Engine/Plugins/MovieScene/MovieRenderPipeline/Source/MovieRenderPipelineRenderPasses/Private/MoviePipelineDeferredPasses.cpp b/Engine/Plugins/MovieScene/MovieRenderPipeline/Source/MovieRenderPipelineRenderPasses/Private/MoviePipelineDeferredPasses.cpp
--- a/Engine/Plugins/MovieScene/MovieRenderPipeline/Source/MovieRenderPipelineRenderPasses/Private/MoviePipelineDeferredPasses.cpp
+++ b/Engine/Plugins/MovieScene/MovieRenderPipeline/Source/MovieRenderPipelineRenderPasses/Private/MoviePipelineDeferredPasses.cpp
@@ -69,10 +69,10 @@ FIntPoint UMoviePipelineDeferredPassBase::GetEffectiveOutputResolutionForCamera(
 {
 	// Add here the the output resolution for each camera. Now only one size is implemented, obtained from the settings.
 
-	UMoviePipelineMasterConfig* MasterConfig = GetPipeline()->GetPipelineMasterConfig();
+	UMoviePipelinePrimaryConfig* PrimaryConfig = GetPipeline()->GetPipelinePrimaryConfig();
 	UMoviePipelineExecutorShot* CurrentShot = GetPipeline()->GetActiveShotList()[GetPipeline()->GetCurrentShotIndex()];
 
-	const FIntPoint OutputResolution = UMoviePipelineBlueprintLibrary::GetEffectiveOutputResolution(MasterConfig, CurrentShot);
+	const FIntPoint OutputResolution = UMoviePipelineBlueprintLibrary::GetEffectiveOutputResolution(PrimaryConfig, CurrentShot);
 
 	return OutputResolution;
 }
@@ -80,12 +80,12 @@ FIntPoint UMoviePipelineDeferredPassBase::GetEffectiveOutputResolutionForCamera(
 FMoviePipelineRenderPassMetrics UMoviePipelineDeferredPassBase::GetRenderPassMetricsForCamera(const int32 InCameraIndex, const FMoviePipelineRenderPassMetrics& InSampleState) const
 {
 	// Add per-camera custom backbuffer size support here.
-	UMoviePipelineMasterConfig* MasterConfig = GetPipeline()->GetPipelineMasterConfig();
+	UMoviePipelinePrimaryConfig* PrimaryConfig = GetPipeline()->GetPipelinePrimaryConfig();
 	UMoviePipelineExecutorShot* CurrentShot = GetPipeline()->GetActiveShotList()[GetPipeline()->GetCurrentShotIndex()];
-	check(MasterConfig);
+	check(PrimaryConfig);
 	check(CurrentShot);
 
-	return UE::MoviePipeline::GetRenderPassMetrics(MasterConfig, CurrentShot, InSampleState, GetEffectiveOutputResolutionForCamera(InCameraIndex));
+	return UE::MoviePipeline::GetRenderPassMetrics(PrimaryConfig, CurrentShot, InSampleState, GetEffectiveOutputResolutionForCamera(InCameraIndex));
 }
 
 int32 UMoviePipelineDeferredPassBase::GetNumCamerasToRender() const
@@ -844,7 +844,7 @@ void UMoviePipelineDeferredPassBase::PostRendererSubmission(const FMoviePipeline
 	if(PlayerCameraManager && PlayerCameraManager->GetCameraCacheView().bConstrainAspectRatio)
 	{
 		const FMinimalViewInfo CameraCache = PlayerCameraManager->GetCameraCacheView();
-		UMoviePipelineOutputSetting* OutputSettings = GetPipeline()->GetPipelineMasterConfig()->FindSetting<UMoviePipelineOutputSetting>();
+		UMoviePipelineOutputSetting* OutputSettings = GetPipeline()->GetPipelinePrimaryConfig()->FindSetting<UMoviePipelineOutputSetting>();
 		check(OutputSettings);
 		
 		// Taking overscan into account.
diff --git a/Engine/Plugins/MovieScene/MovieRenderPipeline/Source/MovieRenderPipelineRenderPasses/Private/MoviePipelineEXROutput.cpp b/Engine/Plugins/MovieScene/MovieRenderPipeline/Source/MovieRenderPipelineRenderPasses/Private/MoviePipelineEXROutput.cpp
--- a/Engine/Plugins/MovieScene/MovieRenderPipeline/Source/MovieRenderPipelineRenderPasses/Private/MoviePipelineEXROutput.cpp
+++ b/Engine/Plugins/MovieScene/MovieRenderPipeline/Source/MovieRenderPipelineRenderPasses/Private/MoviePipelineEXROutput.cpp
@@ -14,7 +14,7 @@
 #include "ImageWriteQueue.h"
 #include "MoviePipeline.h"
 #include "MoviePipelineImageQuantization.h"
-#include "MoviePipelineMasterConfig.h"
+#include "MoviePipelinePrimaryConfig.h"
 #include "IOpenExrRTTIModule.h"
 #include "Modules/ModuleManager.h"
 #include "MoviePipelineUtils.h"
@@ -401,7 +401,7 @@ void UMoviePipelineImageSequenceOutput_EXR::OnReceiveImageDataImpl(FMoviePipelin
 	FModuleManager::Get().LoadModule(RTTIExtensionModuleName);
 
 
-	UMoviePipelineOutputSetting* OutputSettings = GetPipeline()->GetPipelineMasterConfig()->FindSetting<UMoviePipelineOutputSetting>();
+	UMoviePipelineOutputSetting* OutputSettings = GetPipeline()->GetPipelinePrimaryConfig()->FindSetting<UMoviePipelineOutputSetting>();
 	check(OutputSettings);
 
 	// EXR only supports one resolution per file, but in certain scenarios we can get layers with different resolutions. To solve this, we're
diff --git a/Engine/Plugins/MovieScene/MovieRenderPipeline/Source/MovieRenderPipelineRenderPasses/Private/MoviePipelineImagePassBase.cpp b/Engine/Plugins/MovieScene/MovieRenderPipeline/Source/MovieRenderPipelineRenderPasses/Private/MoviePipelineImagePassBase.cpp
--- a/Engine/Plugins/MovieScene/MovieRenderPipeline/Source/MovieRenderPipelineRenderPasses/Private/MoviePipelineImagePassBase.cpp
+++ b/Engine/Plugins/MovieScene/MovieRenderPipeline/Source/MovieRenderPipelineRenderPasses/Private/MoviePipelineImagePassBase.cpp
@@ -12,7 +12,7 @@
 #include "MoviePipelineViewFamilySetting.h"
 #include "MoviePipelineQueue.h"
 #include "LegacyScreenPercentageDriver.h"
-#include "MoviePipelineMasterConfig.h"
+#include "MoviePipelinePrimaryConfig.h"
 #include "MoviePipelineGameOverrideSetting.h"
 #include "EngineModule.h"
 #include "Engine/LocalPlayer.h"
@@ -238,7 +238,7 @@ TSharedPtr<FSceneViewFamilyContext> UMoviePipelineImagePassBase::CalculateViewFa
 
 	// Override the Motion Blur settings since these are controlled by the movie pipeline.
 	{
-		FFrameRate OutputFrameRate = GetPipeline()->GetPipelineMasterConfig()->GetEffectiveFrameRate(GetPipeline()->GetTargetSequence());
+		FFrameRate OutputFrameRate = GetPipeline()->GetPipelinePrimaryConfig()->GetEffectiveFrameRate(GetPipeline()->GetTargetSequence());
 
 		// We need to inversly scale the target FPS by time dilation to counteract slowmo. If scaling isn't applied then motion blur length
 		// stays the same length despite the smaller delta time and the blur ends up too long.
diff --git a/Engine/Plugins/MovieScene/MovieRenderPipeline/Source/MovieRenderPipelineRenderPasses/Private/MoviePipelineImageSequenceOutput.cpp b/Engine/Plugins/MovieScene/MovieRenderPipeline/Source/MovieRenderPipelineRenderPasses/Private/MoviePipelineImageSequenceOutput.cpp
--- a/Engine/Plugins/MovieScene/MovieRenderPipeline/Source/MovieRenderPipelineRenderPasses/Private/MoviePipelineImageSequenceOutput.cpp
+++ b/Engine/Plugins/MovieScene/MovieRenderPipeline/Source/MovieRenderPipelineRenderPasses/Private/MoviePipelineImageSequenceOutput.cpp
@@ -7,7 +7,7 @@
 #include "ImageWriteQueue.h"
 #include "MoviePipeline.h"
 #include "ImageWriteStream.h"
-#include "MoviePipelineMasterConfig.h"
+#include "MoviePipelinePrimaryConfig.h"
 #include "MovieRenderTileImage.h"
 #include "MovieRenderOverlappedImage.h"
 #include "MovieRenderPipelineCoreModule.h"
@@ -88,10 +88,10 @@ void UMoviePipelineImageSequenceOutputBase::OnReceiveImageDataImpl(FMoviePipelin
 	MoviePipeline::GetPassCompositeData(InMergedOutputFrame, CompositedPasses);
 
 
-	UMoviePipelineOutputSetting* OutputSettings = GetPipeline()->GetPipelineMasterConfig()->FindSetting<UMoviePipelineOutputSetting>();
+	UMoviePipelineOutputSetting* OutputSettings = GetPipeline()->GetPipelinePrimaryConfig()->FindSetting<UMoviePipelineOutputSetting>();
 	check(OutputSettings);
 
-	UMoviePipelineColorSetting* ColorSetting = GetPipeline()->GetPipelineMasterConfig()->FindSetting<UMoviePipelineColorSetting>();
+	UMoviePipelineColorSetting* ColorSetting = GetPipeline()->GetPipelinePrimaryConfig()->FindSetting<UMoviePipelineColorSetting>();
 
 	FString OutputDirectory = OutputSettings->OutputDirectory.Path;
 
diff --git a/Engine/Plugins/MovieScene/MovieRenderPipeline/Source/MovieRenderPipelineRenderPasses/Private/MoviePipelineWaveOutput.cpp b/Engine/Plugins/MovieScene/MovieRenderPipeline/Source/MovieRenderPipelineRenderPasses/Private/MoviePipelineWaveOutput.cpp
--- a/Engine/Plugins/MovieScene/MovieRenderPipeline/Source/MovieRenderPipelineRenderPasses/Private/MoviePipelineWaveOutput.cpp
+++ b/Engine/Plugins/MovieScene/MovieRenderPipeline/Source/MovieRenderPipelineRenderPasses/Private/MoviePipelineWaveOutput.cpp
@@ -6,7 +6,7 @@
 #include "SampleBuffer.h"
 #include "MoviePipelineOutputSetting.h"
 #include "MoviePipelineSetting.h"
-#include "MoviePipelineMasterConfig.h"
+#include "MoviePipelinePrimaryConfig.h"
 #include "MovieRenderPipelineCoreModule.h"
 #include "Misc/CommandLine.h"
 #include "Misc/Parse.h"
@@ -91,7 +91,7 @@ void UMoviePipelineWaveOutput::BeginFinalizeImpl()
 		return;
 	}
 
-	UMoviePipelineOutputSetting* OutputSetting = GetPipeline()->GetPipelineMasterConfig()->FindSetting<UMoviePipelineOutputSetting>();
+	UMoviePipelineOutputSetting* OutputSetting = GetPipeline()->GetPipelinePrimaryConfig()->FindSetting<UMoviePipelineOutputSetting>();
 	FString OutputFilename = FileNameFormatOverride.Len() > 0 ? FileNameFormatOverride : OutputSetting->FileNameFormat;
 	FString FileNameFormatString = OutputSetting->OutputDirectory.Path / OutputFilename;
 
diff --git a/Engine/Plugins/MovieScene/MovieRenderPipeline/Source/MovieRenderPipelineSettings/Private/MoviePipelineBurnInSetting.cpp b/Engine/Plugins/MovieScene/MovieRenderPipeline/Source/MovieRenderPipelineSettings/Private/MoviePipelineBurnInSetting.cpp
--- a/Engine/Plugins/MovieScene/MovieRenderPipeline/Source/MovieRenderPipelineSettings/Private/MoviePipelineBurnInSetting.cpp
+++ b/Engine/Plugins/MovieScene/MovieRenderPipeline/Source/MovieRenderPipelineSettings/Private/MoviePipelineBurnInSetting.cpp
@@ -7,7 +7,7 @@
 #include "MoviePipelineBurnInWidget.h"
 #include "MoviePipelineOutputSetting.h"
 #include "MoviePipelineCameraSetting.h"
-#include "MoviePipelineMasterConfig.h"
+#include "MoviePipelinePrimaryConfig.h"
 #include "MoviePipelineBlueprintLibrary.h"
 #include "MoviePipeline.h"
 #include "Engine/TextureRenderTarget2D.h"
@@ -114,7 +114,7 @@ void UMoviePipelineBurnInSetting::SetupImpl(const MoviePipeline::FMoviePipelineR
 		return;
 	}
 
-	OutputResolution = UMoviePipelineBlueprintLibrary::GetEffectiveOutputResolution(GetPipeline()->GetPipelineMasterConfig(), GetPipeline()->GetActiveShotList()[GetPipeline()->GetCurrentShotIndex()]);
+	OutputResolution = UMoviePipelineBlueprintLibrary::GetEffectiveOutputResolution(GetPipeline()->GetPipelinePrimaryConfig(), GetPipeline()->GetActiveShotList()[GetPipeline()->GetCurrentShotIndex()]);
 	int32 MaxResolution = GetMax2DTextureDimension();
 	if (OutputResolution.X > MaxResolution || OutputResolution.Y > MaxResolution)
 	{
diff --git a/Engine/Plugins/MovieScene/MovieRenderPipeline/Source/MovieRenderPipelineSettings/Private/MoviePipelineWidgetRenderSetting.cpp b/Engine/Plugins/MovieScene/MovieRenderPipeline/Source/MovieRenderPipelineSettings/Private/MoviePipelineWidgetRenderSetting.cpp
--- a/Engine/Plugins/MovieScene/MovieRenderPipeline/Source/MovieRenderPipelineSettings/Private/MoviePipelineWidgetRenderSetting.cpp
+++ b/Engine/Plugins/MovieScene/MovieRenderPipeline/Source/MovieRenderPipelineSettings/Private/MoviePipelineWidgetRenderSetting.cpp
@@ -6,7 +6,7 @@
 #include "MoviePipelineBurnInWidget.h"
 #include "MoviePipelineOutputSetting.h"
 #include "MoviePipelineCameraSetting.h"
-#include "MoviePipelineMasterConfig.h"
+#include "MoviePipelinePrimaryConfig.h"
 #include "MoviePipelineBlueprintLibrary.h"
 #include "MoviePipeline.h"
 #include "Engine/TextureRenderTarget2D.h"
@@ -104,7 +104,7 @@ void UMoviePipelineWidgetRenderer::SetupImpl(const MoviePipeline::FMoviePipeline
 
 	bool bInForceLinearGamma = false;
 
-	FIntPoint OutputResolution = UMoviePipelineBlueprintLibrary::GetEffectiveOutputResolution(GetPipeline()->GetPipelineMasterConfig(), GetPipeline()->GetActiveShotList()[GetPipeline()->GetCurrentShotIndex()]);
+	FIntPoint OutputResolution = UMoviePipelineBlueprintLibrary::GetEffectiveOutputResolution(GetPipeline()->GetPipelinePrimaryConfig(), GetPipeline()->GetActiveShotList()[GetPipeline()->GetCurrentShotIndex()]);
 
 	int32 MaxResolution = GetMax2DTextureDimension();
 	if (OutputResolution.X > MaxResolution || OutputResolution.Y > MaxResolution)
diff --git a/Engine/Plugins/MovieScene/ReplayTracks/Source/ReplayTracks/Private/Systems/MovieSceneReplaySystem.cpp b/Engine/Plugins/MovieScene/ReplayTracks/Source/ReplayTracks/Private/Systems/MovieSceneReplaySystem.cpp
--- a/Engine/Plugins/MovieScene/ReplayTracks/Source/ReplayTracks/Private/Systems/MovieSceneReplaySystem.cpp
+++ b/Engine/Plugins/MovieScene/ReplayTracks/Source/ReplayTracks/Private/Systems/MovieSceneReplaySystem.cpp
@@ -6,7 +6,7 @@
 #include "Engine/EngineTypes.h"
 #include "Engine/World.h"
 #include "EntitySystem/MovieSceneEntitySystemTask.h"
-#include "EntitySystem/MovieSceneMasterInstantiatorSystem.h"
+#include "EntitySystem/MovieSceneRootInstantiatorSystem.h"
 #include "Evaluation/MovieSceneEvaluationTemplateInstance.h"
 #include "GameDelegates.h"
 #include "GameFramework/GameState.h"
@@ -72,7 +72,7 @@ UMovieSceneReplaySystem::UMovieSceneReplaySystem(const FObjectInitializer& ObjIn
 
 	if (HasAnyFlags(RF_ClassDefaultObject))
 	{
-		DefineImplicitPrerequisite(UMovieSceneMasterInstantiatorSystem::StaticClass(), GetClass());
+		DefineImplicitPrerequisite(UMovieSceneRootInstantiatorSystem::StaticClass(), GetClass());
 	}
 }
 
diff --git a/Engine/Plugins/MovieScene/ReplayTracks/Source/ReplayTracksEditor/Private/TrackEditors/ReplayTrackEditor.cpp b/Engine/Plugins/MovieScene/ReplayTracks/Source/ReplayTracksEditor/Private/TrackEditors/ReplayTrackEditor.cpp
--- a/Engine/Plugins/MovieScene/ReplayTracks/Source/ReplayTracksEditor/Private/TrackEditors/ReplayTrackEditor.cpp
+++ b/Engine/Plugins/MovieScene/ReplayTracks/Source/ReplayTracksEditor/Private/TrackEditors/ReplayTrackEditor.cpp
@@ -90,7 +90,7 @@ bool FReplayTrackEditor::HandleAddReplayTrackMenuEntryCanExecute() const
 {
 	UMovieScene* FocusedMovieScene = GetFocusedMovieScene();
 
-	return ((FocusedMovieScene != nullptr) && (FocusedMovieScene->FindMasterTrack<UMovieSceneReplayTrack>() == nullptr));
+	return ((FocusedMovieScene != nullptr) && (FocusedMovieScene->FindTrack<UMovieSceneReplayTrack>() == nullptr));
 }
 
 void FReplayTrackEditor::HandleAddReplayTrackMenuEntryExecute()
@@ -189,14 +189,14 @@ UMovieSceneReplayTrack* FReplayTrackEditor::FindOrCreateReplayTrack(bool* bWasCr
 		return nullptr;
 	}
 
-	UMovieSceneReplayTrack* ReplayTrack = FocusedMovieScene->FindMasterTrack<UMovieSceneReplayTrack>();
+	UMovieSceneReplayTrack* ReplayTrack = FocusedMovieScene->FindTrack<UMovieSceneReplayTrack>();
 
 	if (ReplayTrack == nullptr)
 	{
 		const FScopedTransaction Transaction(LOCTEXT("AddReplayTrack_Transaction", "Add Replay Track"));
 		FocusedMovieScene->Modify();
 		
-		ReplayTrack = FocusedMovieScene->AddMasterTrack<UMovieSceneReplayTrack>();
+		ReplayTrack = FocusedMovieScene->AddTrack<UMovieSceneReplayTrack>();
 
 		if (bWasCreated)
 		{
diff --git a/Engine/Plugins/MovieScene/SequencerScripting/Source/SequencerScripting/Private/ExtensionLibraries/MovieSceneFolderExtensions.cpp b/Engine/Plugins/MovieScene/SequencerScripting/Source/SequencerScripting/Private/ExtensionLibraries/MovieSceneFolderExtensions.cpp
--- a/Engine/Plugins/MovieScene/SequencerScripting/Source/SequencerScripting/Private/ExtensionLibraries/MovieSceneFolderExtensions.cpp
+++ b/Engine/Plugins/MovieScene/SequencerScripting/Source/SequencerScripting/Private/ExtensionLibraries/MovieSceneFolderExtensions.cpp
@@ -107,50 +107,50 @@ bool UMovieSceneFolderExtensions::RemoveChildFolder(UMovieSceneFolder* TargetFol
 	return true;
 }
 
-TArray<UMovieSceneTrack*> UMovieSceneFolderExtensions::GetChildMasterTracks(UMovieSceneFolder* Folder)
+TArray<UMovieSceneTrack*> UMovieSceneFolderExtensions::GetChildTracks(UMovieSceneFolder* Folder)
 {
 	if (!Folder)
 	{
-		FFrame::KismetExecutionMessage(TEXT("Cannot call GetChildMasterTracks with a null folder"), ELogVerbosity::Error);
+		FFrame::KismetExecutionMessage(TEXT("Cannot call GetChildTracks with a null folder"), ELogVerbosity::Error);
 		return TArray<UMovieSceneTrack*>();
 	}
 
-	return Folder->GetChildMasterTracks();
+	return Folder->GetChildTracks();
 }
 
-bool UMovieSceneFolderExtensions::AddChildMasterTrack(UMovieSceneFolder* Folder, UMovieSceneTrack* InMasterTrack)
+bool UMovieSceneFolderExtensions::AddChildTrack(UMovieSceneFolder* Folder, UMovieSceneTrack* InTrack)
 {
 	if (!Folder)
 	{
-		FFrame::KismetExecutionMessage(TEXT("Cannot call AddChildMasterTrack with a null folder"), ELogVerbosity::Error);
+		FFrame::KismetExecutionMessage(TEXT("Cannot call AddChildTrack with a null folder"), ELogVerbosity::Error);
 		return false;
 	}
 
-	if (!InMasterTrack)
+	if (!InTrack)
 	{
-		FFrame::KismetExecutionMessage(TEXT("Cannot call AddChildMasterTrack with a null master track"), ELogVerbosity::Error);
+		FFrame::KismetExecutionMessage(TEXT("Cannot call AddChildTrack with a null track"), ELogVerbosity::Error);
 		return false;
 	}
 	
-	Folder->AddChildMasterTrack(InMasterTrack);
+	Folder->AddChildTrack(InTrack);
 	return true;
 }
 
-bool UMovieSceneFolderExtensions::RemoveChildMasterTrack(UMovieSceneFolder* Folder, UMovieSceneTrack* InMasterTrack)
+bool UMovieSceneFolderExtensions::RemoveChildTrack(UMovieSceneFolder* Folder, UMovieSceneTrack* InTrack)
 {
 	if (!Folder)
 	{
-		FFrame::KismetExecutionMessage(TEXT("Cannot call RemoveChildMasterTrack with a null folder"), ELogVerbosity::Error);
+		FFrame::KismetExecutionMessage(TEXT("Cannot call RemoveChildTrack with a null folder"), ELogVerbosity::Error);
 		return false;
 	}
 
-	if (!InMasterTrack)
+	if (!InTrack)
 	{
-		FFrame::KismetExecutionMessage(TEXT("Cannot call RemoveChildMasterTrack with a null master track"), ELogVerbosity::Error);
+		FFrame::KismetExecutionMessage(TEXT("Cannot call RemoveChildTrack with a null track"), ELogVerbosity::Error);
 		return false;
 	}
 
-	Folder->RemoveChildMasterTrack(InMasterTrack);
+	Folder->RemoveChildTrack(InTrack);
 	return true;
 }
 
diff --git a/Engine/Plugins/MovieScene/SequencerScripting/Source/SequencerScripting/Private/ExtensionLibraries/MovieSceneSectionExtensions.cpp b/Engine/Plugins/MovieScene/SequencerScripting/Source/SequencerScripting/Private/ExtensionLibraries/MovieSceneSectionExtensions.cpp
--- a/Engine/Plugins/MovieScene/SequencerScripting/Source/SequencerScripting/Private/ExtensionLibraries/MovieSceneSectionExtensions.cpp
+++ b/Engine/Plugins/MovieScene/SequencerScripting/Source/SequencerScripting/Private/ExtensionLibraries/MovieSceneSectionExtensions.cpp
@@ -623,9 +623,9 @@ TArray<UMovieSceneScriptingChannel*> UMovieSceneSectionExtensions::GetChannelsBy
 bool GetSubSectionChain(UMovieSceneSubSection* InSubSection, UMovieSceneSequence* ParentSequence, TArray<UMovieSceneSubSection*>& SubSectionChain)
 {
 	UMovieScene* ParentMovieScene = ParentSequence->GetMovieScene();
-	for (UMovieSceneTrack* MasterTrack : ParentMovieScene->GetMasterTracks())
+	for (UMovieSceneTrack* Track : ParentMovieScene->GetTracks())
 	{
-		for (UMovieSceneSection* Section : MasterTrack->GetAllSections())
+		for (UMovieSceneSection* Section : Track->GetAllSections())
 		{
 			if (Section == InSubSection)
 			{
diff --git a/Engine/Plugins/MovieScene/SequencerScripting/Source/SequencerScripting/Private/ExtensionLibraries/MovieSceneSequenceExtensions.cpp b/Engine/Plugins/MovieScene/SequencerScripting/Source/SequencerScripting/Private/ExtensionLibraries/MovieSceneSequenceExtensions.cpp
--- a/Engine/Plugins/MovieScene/SequencerScripting/Source/SequencerScripting/Private/ExtensionLibraries/MovieSceneSequenceExtensions.cpp
+++ b/Engine/Plugins/MovieScene/SequencerScripting/Source/SequencerScripting/Private/ExtensionLibraries/MovieSceneSequenceExtensions.cpp
@@ -38,11 +38,11 @@ UMovieScene* UMovieSceneSequenceExtensions::GetMovieScene(UMovieSceneSequence* S
 	return Sequence->GetMovieScene();
 }
 
-TArray<UMovieSceneTrack*> UMovieSceneSequenceExtensions::GetMasterTracks(UMovieSceneSequence* Sequence)
+TArray<UMovieSceneTrack*> UMovieSceneSequenceExtensions::GetTracks(UMovieSceneSequence* Sequence)
 {
 	if (!Sequence)
 	{
-		FFrame::KismetExecutionMessage(TEXT("Cannot call GetMasterTracks on a null sequence"), ELogVerbosity::Error);
+		FFrame::KismetExecutionMessage(TEXT("Cannot call GetTracks on a null sequence"), ELogVerbosity::Error);
 		return TArray<UMovieSceneTrack*>();
 	}
 
@@ -51,7 +51,7 @@ TArray<UMovieSceneTrack*> UMovieSceneSequenceExtensions::GetMasterTracks(UMovieS
 	UMovieScene* MovieScene = GetMovieScene(Sequence);
 	if (MovieScene)
 	{
-		Tracks = MovieScene->GetMasterTracks();
+		Tracks = MovieScene->GetTracks();
 
 		if (UMovieSceneTrack* CameraCutTrack = MovieScene->GetCameraCutTrack())
 		{
@@ -62,11 +62,11 @@ TArray<UMovieSceneTrack*> UMovieSceneSequenceExtensions::GetMasterTracks(UMovieS
 	return Tracks;
 }
 
-TArray<UMovieSceneTrack*> UMovieSceneSequenceExtensions::FindMasterTracksByType(UMovieSceneSequence* Sequence, TSubclassOf<UMovieSceneTrack> TrackType)
+TArray<UMovieSceneTrack*> UMovieSceneSequenceExtensions::FindTracksByType(UMovieSceneSequence* Sequence, TSubclassOf<UMovieSceneTrack> TrackType)
 {
 	if (!Sequence)
 	{
-		FFrame::KismetExecutionMessage(TEXT("Cannot call FindMasterTracksByType on a null sequence"), ELogVerbosity::Error);
+		FFrame::KismetExecutionMessage(TEXT("Cannot call FindTracksByType on a null sequence"), ELogVerbosity::Error);
 		return TArray<UMovieSceneTrack*>();
 	}
 
@@ -76,9 +76,9 @@ TArray<UMovieSceneTrack*> UMovieSceneSequenceExtensions::FindMasterTracksByType(
 	if (MovieScene && DesiredClass)
 	{
 		bool bExactMatch = false;
-		TArray<UMovieSceneTrack*> MatchedTracks = FilterTracks(MovieScene->GetMasterTracks(), TrackType.Get(), bExactMatch);
+		TArray<UMovieSceneTrack*> MatchedTracks = FilterTracks(MovieScene->GetTracks(), TrackType.Get(), bExactMatch);
 
-		// Have to check camera cut tracks separately since they're not in the master tracks array (why?)
+		// Have to check camera cut tracks separately since they're not in the root tracks array (why?)
 		UMovieSceneTrack* CameraCutTrack = MovieScene->GetCameraCutTrack();
 		if (CameraCutTrack && CameraCutTrack->GetClass()->IsChildOf(DesiredClass))
 		{
@@ -91,11 +91,11 @@ TArray<UMovieSceneTrack*> UMovieSceneSequenceExtensions::FindMasterTracksByType(
 	return TArray<UMovieSceneTrack*>();
 }
 
-TArray<UMovieSceneTrack*> UMovieSceneSequenceExtensions::FindMasterTracksByExactType(UMovieSceneSequence* Sequence, TSubclassOf<UMovieSceneTrack> TrackType)
+TArray<UMovieSceneTrack*> UMovieSceneSequenceExtensions::FindTracksByExactType(UMovieSceneSequence* Sequence, TSubclassOf<UMovieSceneTrack> TrackType)
 {
 	if (!Sequence)
 	{
-		FFrame::KismetExecutionMessage(TEXT("Cannot call FindMasterTracksByExactType on a null sequence"), ELogVerbosity::Error);
+		FFrame::KismetExecutionMessage(TEXT("Cannot call FindTracksByExactType on a null sequence"), ELogVerbosity::Error);
 		return TArray<UMovieSceneTrack*>();
 	}
 
@@ -105,9 +105,9 @@ TArray<UMovieSceneTrack*> UMovieSceneSequenceExtensions::FindMasterTracksByExact
 	if (MovieScene && DesiredClass)
 	{
 		bool bExactMatch = true;
-		TArray<UMovieSceneTrack*> MatchedTracks = FilterTracks(MovieScene->GetMasterTracks(), TrackType.Get(), bExactMatch);
+		TArray<UMovieSceneTrack*> MatchedTracks = FilterTracks(MovieScene->GetTracks(), TrackType.Get(), bExactMatch);
 
-		// Have to check camera cut tracks separately since they're not in the master tracks array (why?)
+		// Have to check camera cut tracks separately since they're not in the root tracks array (why?)
 		UMovieSceneTrack* CameraCutTrack = MovieScene->GetCameraCutTrack();
 		if (CameraCutTrack && CameraCutTrack->GetClass() == DesiredClass)
 		{
@@ -120,15 +120,15 @@ TArray<UMovieSceneTrack*> UMovieSceneSequenceExtensions::FindMasterTracksByExact
 	return TArray<UMovieSceneTrack*>();
 }
 
-UMovieSceneTrack* UMovieSceneSequenceExtensions::AddMasterTrack(UMovieSceneSequence* Sequence, TSubclassOf<UMovieSceneTrack> TrackType)
+UMovieSceneTrack* UMovieSceneSequenceExtensions::AddTrack(UMovieSceneSequence* Sequence, TSubclassOf<UMovieSceneTrack> TrackType)
 {
 	if (!Sequence)
 	{
-		FFrame::KismetExecutionMessage(TEXT("Cannot call AddMasterTrack on a null sequence"), ELogVerbosity::Error);
+		FFrame::KismetExecutionMessage(TEXT("Cannot call AddTrack on a null sequence"), ELogVerbosity::Error);
 		return nullptr;
 	}
 
-	// @todo: sequencer-python: master track type compatibility with sequence. Currently that's really only loosely defined by track editors, which is not sufficient here.
+	// @todo: sequencer-python: track type compatibility with sequence. Currently that's really only loosely defined by track editors, which is not sufficient here.
 	UMovieScene* MovieScene = GetMovieScene(Sequence);
 	if (MovieScene)
 	{
@@ -138,32 +138,32 @@ UMovieSceneTrack* UMovieSceneSequenceExtensions::AddMasterTrack(UMovieSceneSeque
 		}
 		else
 		{
-			return MovieScene->AddMasterTrack(TrackType);
+			return MovieScene->AddTrack(TrackType);
 		}
 	}
 
 	return nullptr;
 }
 
-bool UMovieSceneSequenceExtensions::RemoveMasterTrack(UMovieSceneSequence* Sequence, UMovieSceneTrack* MasterTrack)
+bool UMovieSceneSequenceExtensions::RemoveTrack(UMovieSceneSequence* Sequence, UMovieSceneTrack* Track)
 {
 	if (!Sequence)
 	{
-		FFrame::KismetExecutionMessage(TEXT("Cannot call RemoveMasterTrack on a null sequence"), ELogVerbosity::Error);
+		FFrame::KismetExecutionMessage(TEXT("Cannot call RemoveTrack on a null sequence"), ELogVerbosity::Error);
 		return false;
 	}
 
 	UMovieScene* MovieScene = GetMovieScene(Sequence);
 	if (MovieScene)
 	{
-		if (MasterTrack->IsA(UMovieSceneCameraCutTrack::StaticClass()))
+		if (Track->IsA(UMovieSceneCameraCutTrack::StaticClass()))
 		{
 			MovieScene->RemoveCameraCutTrack();
 			return true;
 		}
 		else
 		{
-			return MovieScene->RemoveMasterTrack(*MasterTrack);
+			return MovieScene->RemoveTrack(*Track);
 		}
 	}
 
@@ -842,9 +842,9 @@ TArray<UObject*> UMovieSceneSequenceExtensions::LocateBoundObjects(UMovieSceneSe
 	return Result;
 }
 
-FMovieSceneObjectBindingID UMovieSceneSequenceExtensions::MakeBindingID(UMovieSceneSequence* MasterSequence, const FMovieSceneBindingProxy& InBinding, EMovieSceneObjectBindingSpace Space)
+FMovieSceneObjectBindingID UMovieSceneSequenceExtensions::MakeBindingID(UMovieSceneSequence* RootSequence, const FMovieSceneBindingProxy& InBinding, EMovieSceneObjectBindingSpace Space)
 {
-	if (!MasterSequence)
+	if (!RootSequence)
 	{
 		FFrame::KismetExecutionMessage(TEXT("Cannot call MakeBindingID on a null sequence"), ELogVerbosity::Error);
 		return FMovieSceneObjectBindingID();
@@ -852,7 +852,7 @@ FMovieSceneObjectBindingID UMovieSceneSequenceExtensions::MakeBindingID(UMovieSc
 
 	// This function was kinda flawed before - when ::Local was passed for the Space parameter,
 	// and the sub sequence ID could not be found it would always fall back to a binding for ::Root without any Sequence ID
-	FMovieSceneObjectBindingID BindingID = GetPortableBindingID(MasterSequence, MasterSequence, InBinding);
+	FMovieSceneObjectBindingID BindingID = GetPortableBindingID(RootSequence, RootSequence, InBinding);
 	if (Space == EMovieSceneObjectBindingSpace::Root)
 	{
 		BindingID.ReinterpretAsFixed();
@@ -865,16 +865,16 @@ FMovieSceneObjectBindingID UMovieSceneSequenceExtensions::GetBindingID(const FMo
 	return UE::MovieScene::FRelativeObjectBindingID(InBinding.BindingID);
 }
 
-FMovieSceneObjectBindingID UMovieSceneSequenceExtensions::GetPortableBindingID(UMovieSceneSequence* MasterSequence, UMovieSceneSequence* DestinationSequence, const FMovieSceneBindingProxy& InBinding)
+FMovieSceneObjectBindingID UMovieSceneSequenceExtensions::GetPortableBindingID(UMovieSceneSequence* RootSequence, UMovieSceneSequence* DestinationSequence, const FMovieSceneBindingProxy& InBinding)
 {
-	if (!MasterSequence || !DestinationSequence || !InBinding.Sequence)
+	if (!RootSequence || !DestinationSequence || !InBinding.Sequence)
 	{
 		FFrame::KismetExecutionMessage(TEXT("Invalid sequence specified."), ELogVerbosity::Error);
 		return FMovieSceneObjectBindingID();
 	}
 
 	// If they are all the same sequence, we're dealing with a local binding - this requires no computation
-	if (MasterSequence == DestinationSequence && MasterSequence == InBinding.Sequence)
+	if (RootSequence == DestinationSequence && RootSequence == InBinding.Sequence)
 	{
 		return UE::MovieScene::FRelativeObjectBindingID(InBinding.BindingID);
 	}
@@ -885,24 +885,24 @@ FMovieSceneObjectBindingID UMovieSceneSequenceExtensions::GetPortableBindingID(U
 	TOptional<FMovieSceneSequenceID> DestinationSequenceID;
 	TOptional<FMovieSceneSequenceID> TargetSequenceID;
 
-	if (MasterSequence == DestinationSequence)
+	if (RootSequence == DestinationSequence)
 	{
 		DestinationSequenceID = MovieSceneSequenceID::Root;
 	}
-	if (MasterSequence == InBinding.Sequence)
+	if (RootSequence == InBinding.Sequence)
 	{
 		TargetSequenceID = MovieSceneSequenceID::Root;
 	}
 
 	// We know that we have at least one sequence ID to find, otherwise we would have entered the ::Local branch above
 	UMovieSceneCompiledDataManager*     CompiledDataManager = UMovieSceneCompiledDataManager::GetPrecompiledData();
-	const FMovieSceneCompiledDataID     DataID              = CompiledDataManager->Compile(MasterSequence);
+	const FMovieSceneCompiledDataID     DataID              = CompiledDataManager->Compile(RootSequence);
 	const FMovieSceneSequenceHierarchy* Hierarchy           = CompiledDataManager->FindHierarchy(DataID);
 
-	// If we have no hierarchy, the supplied MasterSequence does not have any sub-sequences so the callee has given us bogus parameters
+	// If we have no hierarchy, the supplied RootSequence does not have any sub-sequences so the callee has given us bogus parameters
 	if (!Hierarchy)
 	{
-		FFrame::KismetExecutionMessage(*FString::Printf(TEXT("Master Sequence ('%s') does not have any sub-sequences."), *MasterSequence->GetPathName()), ELogVerbosity::Error);
+		FFrame::KismetExecutionMessage(*FString::Printf(TEXT("Root Sequence ('%s') does not have any sub-sequences."), *RootSequence->GetPathName()), ELogVerbosity::Error);
 		return FMovieSceneObjectBindingID();
 	}
 
@@ -931,30 +931,30 @@ FMovieSceneObjectBindingID UMovieSceneSequenceExtensions::GetPortableBindingID(U
 
 	if (!DestinationSequenceID.IsSet())
 	{
-		FFrame::KismetExecutionMessage(*FString::Printf(TEXT("Unable to locate DestinationSequence ('%s') within Master Sequence hierarchy ('%s')."), *DestinationSequence->GetPathName(), *MasterSequence->GetPathName()), ELogVerbosity::Error);
+		FFrame::KismetExecutionMessage(*FString::Printf(TEXT("Unable to locate DestinationSequence ('%s') within Root Sequence hierarchy ('%s')."), *DestinationSequence->GetPathName(), *RootSequence->GetPathName()), ELogVerbosity::Error);
 		return FMovieSceneObjectBindingID();
 	}
 
 	if (!TargetSequenceID.IsSet())
 	{
-		FFrame::KismetExecutionMessage(*FString::Printf(TEXT("Unable to locate Sequence for InBinding ('%s') within Master Sequence hierarchy ('%s')."), *InBinding.Sequence->GetPathName(), *MasterSequence->GetPathName()), ELogVerbosity::Error);
+		FFrame::KismetExecutionMessage(*FString::Printf(TEXT("Unable to locate Sequence for InBinding ('%s') within Root Sequence hierarchy ('%s')."), *InBinding.Sequence->GetPathName(), *RootSequence->GetPathName()), ELogVerbosity::Error);
 		return FMovieSceneObjectBindingID();
 	}
 
 	return UE::MovieScene::FRelativeObjectBindingID(DestinationSequenceID.GetValue(), TargetSequenceID.GetValue(), InBinding.BindingID, Hierarchy);
 }
 
-FMovieSceneBindingProxy UMovieSceneSequenceExtensions::ResolveBindingID(UMovieSceneSequence* MasterSequence, FMovieSceneObjectBindingID InObjectBindingID)
+FMovieSceneBindingProxy UMovieSceneSequenceExtensions::ResolveBindingID(UMovieSceneSequence* RootSequence, FMovieSceneObjectBindingID InObjectBindingID)
 {
-	if (!MasterSequence)
+	if (!RootSequence)
 	{
 		FFrame::KismetExecutionMessage(TEXT("Cannot call ResolveBindingID on a null sequence"), ELogVerbosity::Error);
 		return FMovieSceneBindingProxy();
 	}
 
-	UMovieSceneSequence* Sequence = MasterSequence;
+	UMovieSceneSequence* Sequence = RootSequence;
 
-	FMovieSceneCompiledDataID DataID = UMovieSceneCompiledDataManager::GetPrecompiledData()->Compile(MasterSequence);
+	FMovieSceneCompiledDataID DataID = UMovieSceneCompiledDataManager::GetPrecompiledData()->Compile(RootSequence);
 
 	const FMovieSceneSequenceHierarchy* Hierarchy = UMovieSceneCompiledDataManager::GetPrecompiledData()->FindHierarchy(DataID);
 	if (Hierarchy)
diff --git a/Engine/Plugins/MovieScene/SequencerScripting/Source/SequencerScripting/Public/ExtensionLibraries/MovieSceneFolderExtensions.h b/Engine/Plugins/MovieScene/SequencerScripting/Source/SequencerScripting/Public/ExtensionLibraries/MovieSceneFolderExtensions.h
--- a/Engine/Plugins/MovieScene/SequencerScripting/Source/SequencerScripting/Public/ExtensionLibraries/MovieSceneFolderExtensions.h
+++ b/Engine/Plugins/MovieScene/SequencerScripting/Source/SequencerScripting/Public/ExtensionLibraries/MovieSceneFolderExtensions.h
@@ -89,33 +89,45 @@ public:
 	static bool RemoveChildFolder(UMovieSceneFolder* TargetFolder, UMovieSceneFolder* FolderToRemove);
 
 	/**
-	 * Get the master tracks contained by this folder
+	 * Get the tracks contained by this folder
 	 *
-	 * @param Folder	The folder to get the master tracks of
-	 * @return The master tracks under the given folder
+	 * @param Folder	The folder to get the tracks of
+	 * @return The tracks under the given folder
 	 */
-	UFUNCTION(BlueprintCallable, BlueprintPure, Category = "Sequencer|Sequence", meta=(ScriptMethod))
-	static TArray<UMovieSceneTrack*> GetChildMasterTracks(UMovieSceneFolder* Folder);
+	UFUNCTION(BlueprintCallable, BlueprintPure, Category = "Sequencer|Sequence", meta = (ScriptMethod))
+	static TArray<UMovieSceneTrack*> GetChildTracks(UMovieSceneFolder* Folder);
+	
+	UE_DEPRECATED(5.2, "GetChildMasterTracks is deprecated. Please use GetChildTracks instead")
+	UFUNCTION(BlueprintCallable, BlueprintPure, Category = "Sequencer|Sequence", meta = (ScriptMethod, DeprecatedFunction, DeprecationMessage = "GetChildMasterTracks is deprecated. Please use GetChildTracks instead"))
+	static TArray<UMovieSceneTrack*> GetChildMasterTracks(UMovieSceneFolder* Folder) { return GetChildTracks(Folder); }
 
 	/**
-	 * Add a master track to this folder
+	 * Add a track to this folder
 	 *
-	 * @param Folder			The folder to add a child master track to
-	 * @param InMasterTrack		The master track to add to the folder
+	 * @param Folder			The folder to add a child track to
+	 * @param InTrack		    The track to add to the folder
 	 * @return True if the addition is successful
 	 */
-	UFUNCTION(BlueprintCallable, Category = "Sequencer|Sequence", meta=(ScriptMethod))
-	static bool AddChildMasterTrack(UMovieSceneFolder* Folder, UMovieSceneTrack* InMasterTrack);
+	UFUNCTION(BlueprintCallable, Category = "Sequencer|Sequence", meta = (ScriptMethod))
+	static bool AddChildTrack(UMovieSceneFolder* Folder, UMovieSceneTrack* InTrack);
+	
+	UE_DEPRECATED(5.2, "AddChildMasterTrack is deprecated. Please use AddChildTrack instead")
+	UFUNCTION(BlueprintCallable, Category = "Sequencer|Sequence", meta = (ScriptMethod, DeprecatedFunction, DeprecationMessage = "AddChildMasterTrack is deprecated. Please use AddChildTrack instead"))
+	static bool AddChildMasterTrack(UMovieSceneFolder* Folder, UMovieSceneTrack* InTrack) { return AddChildTrack(Folder, InTrack); }
 
 	/**
-	 * Remove a master track from the given folder
+	 * Remove a track from the given folder
 	 *
 	 * @param Folder			The folder from which to remove a track
-	 * @param InMasterTrack		The track to remove
+	 * @param InTrack		    The track to remove
 	 * @return True if the removal succeeds
 	 */
 	UFUNCTION(BlueprintCallable, Category = "Sequencer|Sequence", meta=(ScriptMethod))
-	static bool RemoveChildMasterTrack(UMovieSceneFolder* Folder, UMovieSceneTrack* InMasterTrack);
+	static bool RemoveChildTrack(UMovieSceneFolder* Folder, UMovieSceneTrack* InTrack);
+
+	UE_DEPRECATED(5.2, "RemoveChildMasterTrack is deprecated. Please use RemoveChildTrack instead")
+	UFUNCTION(BlueprintCallable, Category = "Sequencer|Sequence", meta = (ScriptMethod, DeprecatedFunction, DeprecationMessage = "RemoveChildMasterTrack is deprecated. Please use RemoveChildTrack instead"))
+	static bool RemoveChildMasterTrack(UMovieSceneFolder* Folder, UMovieSceneTrack* InTrack) { return RemoveChildTrack(Folder, InTrack); }
 
 	/**
 	 * Get the object bindings contained by this folder
diff --git a/Engine/Plugins/MovieScene/SequencerScripting/Source/SequencerScripting/Public/ExtensionLibraries/MovieSceneSequenceExtensions.h b/Engine/Plugins/MovieScene/SequencerScripting/Source/SequencerScripting/Public/ExtensionLibraries/MovieSceneSequenceExtensions.h
--- a/Engine/Plugins/MovieScene/SequencerScripting/Source/SequencerScripting/Public/ExtensionLibraries/MovieSceneSequenceExtensions.h
+++ b/Engine/Plugins/MovieScene/SequencerScripting/Source/SequencerScripting/Public/ExtensionLibraries/MovieSceneSequenceExtensions.h
@@ -33,53 +33,73 @@ public:
 	static UMovieScene* GetMovieScene(UMovieSceneSequence* Sequence);
 
 	/**
-	 * Get all master tracks
+	 * Get all tracks
 	 *
 	 * @param Sequence        The sequence to use
-	 * @return An array containing all master tracks in this sequence
+	 * @return An array containing all tracks in this sequence
 	 */
 	UFUNCTION(BlueprintCallable, Category="Sequencer|Sequence", meta=(ScriptMethod))
-	static TArray<UMovieSceneTrack*> GetMasterTracks(UMovieSceneSequence* Sequence);
+	static TArray<UMovieSceneTrack*> GetTracks(UMovieSceneSequence* Sequence);
+
+	UE_DEPRECATED(5.2, "GetMasterTracks is deprecated. Please use GetTracks instead")
+	UFUNCTION(BlueprintCallable, Category = "Sequencer|Sequence", meta = (ScriptMethod, DeprecatedFunction, DeprecationMessage = "GetMasterTracks is deprecated. Please use GetTracks instead"))
+	static TArray<UMovieSceneTrack*> GetMasterTracks(UMovieSceneSequence* Sequence) { return GetTracks(Sequence); }
 
 	/**
-	 * Find all master tracks of the specified type
+	 * Find all tracks of the specified type
 	 *
 	 * @param Sequence        The sequence to use
 	 * @param TrackType     A UMovieSceneTrack class type specifying which types of track to return
 	 * @return An array containing any tracks that match the type specified
 	 */
-	UFUNCTION(BlueprintCallable, Category="Sequencer|Sequence", meta=(ScriptMethod, DeterminesOutputType="TrackType"))
-	static TArray<UMovieSceneTrack*> FindMasterTracksByType(UMovieSceneSequence* Sequence, TSubclassOf<UMovieSceneTrack> TrackType);
+	UFUNCTION(BlueprintCallable, Category = "Sequencer|Sequence", meta = (ScriptMethod, DeterminesOutputType = "TrackType"))
+	static TArray<UMovieSceneTrack*> FindTracksByType(UMovieSceneSequence* Sequence, TSubclassOf<UMovieSceneTrack> TrackType);
+
+	UE_DEPRECATED(5.2, "FindMasterTracksByType is deprecated. Please use FindTracksByType instead")
+	UFUNCTION(BlueprintCallable, Category = "Sequencer|Sequence", meta = (ScriptMethod, DeterminesOutputType = "TrackType", DeprecatedFunction, DeprecationMessage = "FindMasterTracksByType is deprecated. Please use FindTracksByType instead"))
+	static TArray<UMovieSceneTrack*> FindMasterTracksByType(UMovieSceneSequence* Sequence, TSubclassOf<UMovieSceneTrack> TrackType) { return FindTracksByType(Sequence, TrackType); }
 
 	/**
-	 * Find all master tracks of the specified type, not allowing sub-classed types
+	 * Find all tracks of the specified type, not allowing sub-classed types
 	 *
 	 * @param Sequence        The sequence to use
 	 * @param TrackType     A UMovieSceneTrack class type specifying the exact types of track to return
 	 * @return An array containing any tracks that are exactly the same as the type specified
 	 */
 	UFUNCTION(BlueprintCallable, Category="Sequencer|Sequence", meta=(ScriptMethod, DeterminesOutputType="TrackType"))
-	static TArray<UMovieSceneTrack*> FindMasterTracksByExactType(UMovieSceneSequence* Sequence, TSubclassOf<UMovieSceneTrack> TrackType);
+	static TArray<UMovieSceneTrack*> FindTracksByExactType(UMovieSceneSequence* Sequence, TSubclassOf<UMovieSceneTrack> TrackType);
+
+	UE_DEPRECATED(5.2, "FindMasterTracksByExactType is deprecated. Please use FindTracksByExactType instead")
+	UFUNCTION(BlueprintCallable, Category = "Sequencer|Sequence", meta = (ScriptMethod, DeterminesOutputType = "TrackType", DeprecatedFunction, DeprecationMessage = "FindMasterTracksByExactType is deprecated. Please use FindTracksByExactType instead"))
+	static TArray<UMovieSceneTrack*> FindMasterTracksByExactType(UMovieSceneSequence* Sequence, TSubclassOf<UMovieSceneTrack> TrackType) { return FindTracksByExactType(Sequence, TrackType); }
 
 	/**
-	 * Add a new master track of the specified type
+	 * Add a new track of the specified type
 	 *
 	 * @param Sequence        The sequence to use
 	 * @param TrackType     A UMovieSceneTrack class type to create
 	 * @return The newly created track, if successful
 	 */
-	UFUNCTION(BlueprintCallable, Category="Sequencer|Sequence", meta=(ScriptMethod, DeterminesOutputType="TrackType"))
-	static UMovieSceneTrack* AddMasterTrack(UMovieSceneSequence* Sequence, TSubclassOf<UMovieSceneTrack> TrackType);
+	UFUNCTION(BlueprintCallable, Category = "Sequencer|Sequence", meta = (ScriptMethod, DeterminesOutputType = "TrackType"))
+	static UMovieSceneTrack* AddTrack(UMovieSceneSequence* Sequence, TSubclassOf<UMovieSceneTrack> TrackType);
+
+	UE_DEPRECATED(5.2, "AddMasterTrack is deprecated. Please use AddTrack instead")
+	UFUNCTION(BlueprintCallable, Category = "Sequencer|Sequence", meta = (ScriptMethod, DeterminesOutputType = "TrackType", DeprecatedFunction, DeprecationMessage = "AddMasterTrack is deprecated. Please use AddTrack instead"))
+	static UMovieSceneTrack* AddMasterTrack(UMovieSceneSequence* Sequence, TSubclassOf<UMovieSceneTrack> TrackType) { return AddTrack(Sequence, TrackType); }
 
 	/**
-	 * Removes a master track
+	 * Removes a track
 	 *
 	 * @param Sequence        The sequence to use
-	 * @param MasterTrack     The master track to remove
-	 * @return Whether the master track was successfully removed
+	 * @param Track           The track to remove
+	 * @return Whether the track was successfully removed
 	 */
 	UFUNCTION(BlueprintCallable, Category="Sequencer|Sequence", meta=(ScriptMethod))
-	static bool RemoveMasterTrack(UMovieSceneSequence* Sequence, UMovieSceneTrack* MasterTrack);
+	static bool RemoveTrack(UMovieSceneSequence* Sequence, UMovieSceneTrack* Track);
+
+	UE_DEPRECATED(5.2, "RemoveMasterTrack is deprecated. Please use RemoveTrack instead")
+	UFUNCTION(BlueprintCallable, Category = "Sequencer|Sequence", meta = (ScriptMethod, DeprecatedFunction, DeprecationMessage = "RemoveMasterTrack is deprecated. Please use RemoveTrack instead"))
+	static bool RemoveMasterTrack(UMovieSceneSequence* Sequence, UMovieSceneTrack* Track) { return RemoveTrack(Sequence, Track); }
 
 	/**
 	 * Gets this sequence's display rate
@@ -439,14 +459,14 @@ public:
 	/**
 	 * Make a binding id for the given binding in this sequence
 	 *
-	 * @param MasterSequence  The master sequence that contains the sequence
+	 * @param Sequence  The sequence that contains the sequence
 	 * @param Binding The binding proxy to generate the binding id from
 	 * @param Space The object binding space to resolve from (Root or Local)
 	 * @return The new object binding id
 	 */
 	UE_DEPRECATED(5.0, "Please migrate to GetBindingID or GetPortableBindingID depending on use-case.")
 	UFUNCTION(BlueprintCallable, Category = "Sequencer|Sequence", meta = (ScriptMethod, DeprecatedFunction, DeprecationMessage="Please migrate to GetBindingID or GetPortableBindingID depending on use-case."))
-	static FMovieSceneObjectBindingID MakeBindingID(UMovieSceneSequence* MasterSequence, const FMovieSceneBindingProxy& InBinding, EMovieSceneObjectBindingSpace Space = EMovieSceneObjectBindingSpace::Root);
+	static FMovieSceneObjectBindingID MakeBindingID(UMovieSceneSequence* Sequence, const FMovieSceneBindingProxy& InBinding, EMovieSceneObjectBindingSpace Space = EMovieSceneObjectBindingSpace::Root);
 
 
 	/**
@@ -463,25 +483,25 @@ public:
 	/**
 	 * Get a portable binding ID for a binding that resides in a different sequence to the one where this binding will be resolved.
 	 * @note: This function must be used over GetBindingID when the target binding resides in different shots or sub-sequences.
-	 * @note: Only unique instances of sequences within a master sequences are supported
+	 * @note: Only unique instances of sequences within a root sequences are supported
 	 *
-	 * @param MasterSequence The master sequence that contains both the destination sequence (that will resolve the binding ID) and the target sequence that contains the actual binding
+	 * @param RootSequence The root sequence that contains both the destination sequence (that will resolve the binding ID) and the target sequence that contains the actual binding
 	 * @param DestinationSequence The sequence that will own or resolve the resulting binding ID. For example, if the binding ID will be applied to a camera cut section pass the sequence that contains the camera cut track to this parameter.
 	 * @param Binding The target binding to create the ID from
 	 * @return The binding's id
 	 */
 	UFUNCTION(BlueprintCallable, Category = "Sequencer|Sequence", meta = (ScriptMethod))
-	static FMovieSceneObjectBindingID GetPortableBindingID(UMovieSceneSequence* MasterSequence, UMovieSceneSequence* DestinationSequence, const FMovieSceneBindingProxy& InBinding);
+	static FMovieSceneObjectBindingID GetPortableBindingID(UMovieSceneSequence* RootSequence, UMovieSceneSequence* DestinationSequence, const FMovieSceneBindingProxy& InBinding);
 
 	/**
 	 * Make a binding for the given binding ID
 	 *
-	 * @param MasterSequence  The master sequence that contains the sequence
+	 * @param RootSequence  The root sequence that contains the sequence
 	 * @param ObjectBindingID The object binding id that has the guid and the sequence id
 	 * @return The new binding proxy
 	 */
 	UFUNCTION(BlueprintCallable, Category = "Sequencer|Sequence", meta = (ScriptMethod))
-	static FMovieSceneBindingProxy ResolveBindingID(UMovieSceneSequence* MasterSequence, FMovieSceneObjectBindingID InObjectBindingID);
+	static FMovieSceneBindingProxy ResolveBindingID(UMovieSceneSequence* RootSequence, FMovieSceneObjectBindingID InObjectBindingID);
 
 
 	/**
diff --git a/Engine/Plugins/MovieScene/SequencerScripting/Source/SequencerScriptingEditor/Private/SequencerTools.cpp b/Engine/Plugins/MovieScene/SequencerScripting/Source/SequencerScriptingEditor/Private/SequencerTools.cpp
--- a/Engine/Plugins/MovieScene/SequencerScripting/Source/SequencerScriptingEditor/Private/SequencerTools.cpp
+++ b/Engine/Plugins/MovieScene/SequencerScripting/Source/SequencerScriptingEditor/Private/SequencerTools.cpp
@@ -201,7 +201,7 @@ bool ExportFBXInternal(const FSequencerExportFBXParams& InParams, UMovieSceneSeq
 	UMovieSceneSequence* Sequence = InParams.Sequence;
 	UMovieSceneSequence* RootSequence = InParams.RootSequence;
 	TArray<FMovieSceneBindingProxy> BindingProxies = InParams.Bindings;
-	TArray<UMovieSceneTrack*> MasterTracks = InParams.MasterTracks;
+	TArray<UMovieSceneTrack*> Tracks = InParams.Tracks;
 	UFbxExportOption* OverrideOptions = InParams.OverrideOptions;
 	FString FBXFileName = InParams.FBXFileName;
 
@@ -257,7 +257,7 @@ bool ExportFBXInternal(const FSequencerExportFBXParams& InParams, UMovieSceneSeq
 			Player->SetPlaybackPosition(FMovieSceneSequencePlaybackParams(StartTime, EUpdatePositionMethod::Play));
 		}
 
-		bDidExport = MovieSceneToolHelpers::ExportFBX(World, MovieScene, Player, Bindings, MasterTracks, NodeNameAdapter, Template, FBXFileName, RootToLocalTransform);
+		bDidExport = MovieSceneToolHelpers::ExportFBX(World, MovieScene, Player, Bindings, Tracks, NodeNameAdapter, Template, FBXFileName, RootToLocalTransform);
 	}
 
 	Player->Stop();
diff --git a/Engine/Plugins/MovieScene/SequencerScripting/Source/SequencerScriptingEditor/Public/SequencerTools.h b/Engine/Plugins/MovieScene/SequencerScripting/Source/SequencerScriptingEditor/Public/SequencerTools.h
--- a/Engine/Plugins/MovieScene/SequencerScripting/Source/SequencerScriptingEditor/Public/SequencerTools.h
+++ b/Engine/Plugins/MovieScene/SequencerScripting/Source/SequencerScriptingEditor/Public/SequencerTools.h
@@ -74,12 +74,12 @@ struct FSequencerExportFBXParams
 	GENERATED_BODY()
 
 	FSequencerExportFBXParams() {}
-	FSequencerExportFBXParams(UWorld* InWorld, ULevelSequence* InSequence, ULevelSequence* InRootSequence, const TArray<FMovieSceneBindingProxy>& InBindings, const TArray<UMovieSceneTrack*>& InMasterTracks, UFbxExportOption* InOverrideOptions, const FString& InFBXFileName)
+	FSequencerExportFBXParams(UWorld* InWorld, ULevelSequence* InSequence, ULevelSequence* InRootSequence, const TArray<FMovieSceneBindingProxy>& InBindings, const TArray<UMovieSceneTrack*>& InTracks, UFbxExportOption* InOverrideOptions, const FString& InFBXFileName)
 		: World(InWorld)
 		, Sequence(InSequence)
 		, RootSequence(InRootSequence)
 		, Bindings(InBindings)
-		, MasterTracks(InMasterTracks)
+		, Tracks(InTracks)
 		, OverrideOptions(InOverrideOptions)
 		, FBXFileName(InFBXFileName) {}
 	
@@ -96,7 +96,7 @@ struct FSequencerExportFBXParams
 	TArray<FMovieSceneBindingProxy> Bindings;
 	
 	UPROPERTY(BlueprintReadWrite, Category = "Movie Scene")
-	TArray<TObjectPtr<UMovieSceneTrack>> MasterTracks;
+	TArray<TObjectPtr<UMovieSceneTrack>> Tracks;
 	
 	UPROPERTY(BlueprintReadWrite, Category = "Movie Scene")
 	TObjectPtr<UFbxExportOption> OverrideOptions;
@@ -158,13 +158,13 @@ public:
 public:
 
 	/*
-	 * Export Passed in Bindings and Master Tracks to FBX
+	 * Export Passed in Bindings and Tracks to FBX
 	 */
 	UFUNCTION(BlueprintCallable, Category = "Editor Scripting | Sequencer Tools | FBX")
 	static bool ExportLevelSequenceFBX(const FSequencerExportFBXParams& InParams);
 	
 	UE_DEPRECATED(5.1, "Please use ExportLevelSequenceFBX that takes a FSequencerExportFBXParams")
-	static bool ExportLevelSequenceFBX(UWorld* InWorld, ULevelSequence* InSequence, const TArray<FMovieSceneBindingProxy>& InBindings, const TArray<UMovieSceneTrack*>& InMasterTracks, UFbxExportOption* InOverrideOptions, const FString& InFBXFileName) { FSequencerExportFBXParams Params(InWorld, InSequence, InSequence, InBindings, InMasterTracks, InOverrideOptions, InFBXFileName); return ExportLevelSequenceFBX(Params); }
+	static bool ExportLevelSequenceFBX(UWorld* InWorld, ULevelSequence* InSequence, const TArray<FMovieSceneBindingProxy>& InBindings, const TArray<UMovieSceneTrack*>& InTracks, UFbxExportOption* InOverrideOptions, const FString& InFBXFileName) { FSequencerExportFBXParams Params(InWorld, InSequence, InSequence, InBindings, InTracks, InOverrideOptions, InFBXFileName); return ExportLevelSequenceFBX(Params); }
 
 	/*
 	 * Export Passed in Binding as an Anim Seqquence.
diff --git a/Engine/Plugins/MovieScene/TemplateSequence/Source/TemplateSequence/Private/CameraAnimationSequenceSubsystem.cpp b/Engine/Plugins/MovieScene/TemplateSequence/Source/TemplateSequence/Private/CameraAnimationSequenceSubsystem.cpp
--- a/Engine/Plugins/MovieScene/TemplateSequence/Source/TemplateSequence/Private/CameraAnimationSequenceSubsystem.cpp
+++ b/Engine/Plugins/MovieScene/TemplateSequence/Source/TemplateSequence/Private/CameraAnimationSequenceSubsystem.cpp
@@ -9,7 +9,7 @@
 #include "EntitySystem/MovieSceneEntitySystemLinker.h"
 #include "EntitySystem/MovieSceneEntitySystemRunner.h"
 #include "EntitySystem/MovieSceneEntitySystemTask.h"
-#include "EntitySystem/MovieSceneMasterInstantiatorSystem.h"
+#include "EntitySystem/MovieSceneRootInstantiatorSystem.h"
 #include "EntitySystem/MovieSceneBoundSceneComponentInstantiator.h"
 #include "IMovieScenePlayer.h"
 #include "MovieSceneCommonHelpers.h"
diff --git a/Engine/Plugins/MovieScene/TemplateSequence/Source/TemplateSequence/Private/Systems/TemplateSequenceSystem.cpp b/Engine/Plugins/MovieScene/TemplateSequence/Source/TemplateSequence/Private/Systems/TemplateSequenceSystem.cpp
--- a/Engine/Plugins/MovieScene/TemplateSequence/Source/TemplateSequence/Private/Systems/TemplateSequenceSystem.cpp
+++ b/Engine/Plugins/MovieScene/TemplateSequence/Source/TemplateSequence/Private/Systems/TemplateSequenceSystem.cpp
@@ -511,7 +511,7 @@ void UTemplateSequencePropertyScalingEvaluatorSystem::OnRun(FSystemTaskPrerequis
 		.Dispatch_PerEntity<FGatherPropertyScales>(&Linker->EntityManager, InPrerequisites, &Subsequents, this);
 	if (GatherTask)
 	{
-		InPrerequisites.AddMasterTask(GatherTask);
+		InPrerequisites.AddRootTask(GatherTask);
 	}
 
 	// Step 2: If we have any transform property scaling, iterate specifically on them and apply the scale factor
diff --git a/Engine/Plugins/Runtime/GameplayAbilities/Source/GameplayAbilitiesEditor/Private/Sequencer/GameplayCueTrackEditor.cpp b/Engine/Plugins/Runtime/GameplayAbilities/Source/GameplayAbilitiesEditor/Private/Sequencer/GameplayCueTrackEditor.cpp
--- a/Engine/Plugins/Runtime/GameplayAbilities/Source/GameplayAbilitiesEditor/Private/Sequencer/GameplayCueTrackEditor.cpp
+++ b/Engine/Plugins/Runtime/GameplayAbilities/Source/GameplayAbilitiesEditor/Private/Sequencer/GameplayCueTrackEditor.cpp
@@ -173,7 +173,7 @@ void FGameplayCueTrackEditor::AddTracks(TRange<FFrameNumber> SectionTickRange, U
 
 	if (InObjectBindingIDs.Num() == 0)
 	{
-		UMovieSceneGameplayCueTrack* NewMainTrack = FocusedMovieScene->AddMasterTrack<UMovieSceneGameplayCueTrack>();
+		UMovieSceneGameplayCueTrack* NewMainTrack = FocusedMovieScene->AddTrack<UMovieSceneGameplayCueTrack>();
 		AddSectionToTrack(NewMainTrack, SectionTickRange, SectionClass);
 
 		SequencerPtr->OnAddTrack(NewMainTrack, FGuid());
