Sequencer: AddChildFolder adds an existing folder to another folder. It removes it from any other folder but it also needs to remove it from root folders if it is in there.

This fixes a bug in folder scripting where if you call AddChildFolder on a root folder to an existing folder, it will result in being in both folders.

diff --git a/Engine/Plugins/MovieScene/SequencerScripting/Source/SequencerScripting/Private/ExtensionLibraries/MovieSceneSequenceExtensions.cpp b/Engine/Plugins/MovieScene/SequencerScripting/Source/SequencerScripting/Private/ExtensionLibraries/MovieSceneSequenceExtensions.cpp
--- a/Engine/Plugins/MovieScene/SequencerScripting/Source/SequencerScripting/Private/ExtensionLibraries/MovieSceneSequenceExtensions.cpp
+++ b/Engine/Plugins/MovieScene/SequencerScripting/Source/SequencerScripting/Private/ExtensionLibraries/MovieSceneSequenceExtensions.cpp
@@ -1005,7 +1005,7 @@ UMovieSceneFolder* UMovieSceneSequenceExtensions::AddRootFolderToSequence(UMovie
 	if (MovieScene)
 	{
 		MovieScene->Modify();
-		NewFolder = NewObject<UMovieSceneFolder>(MovieScene);
+		NewFolder = NewObject<UMovieSceneFolder>(MovieScene, NAME_None, RF_Transactional);
 		NewFolder->SetFolderName(FName(*NewFolderName));
 		MovieScene->AddRootFolder(NewFolder);
 	}
diff --git a/Engine/Source/Runtime/MovieScene/Private/MovieSceneFolder.cpp b/Engine/Source/Runtime/MovieScene/Private/MovieSceneFolder.cpp
--- a/Engine/Source/Runtime/MovieScene/Private/MovieSceneFolder.cpp
+++ b/Engine/Source/Runtime/MovieScene/Private/MovieSceneFolder.cpp
@@ -64,6 +64,12 @@ void UMovieSceneFolder::AddChildFolder( UMovieSceneFolder* InChildFolder )
 		{
 			MovieSceneFolder->RemoveChildFolder(InChildFolder);
 		}
+
+		if (OwningScene->GetRootFolders().Contains(InChildFolder))
+		{
+			OwningScene->Modify();
+			OwningScene->RemoveRootFolder(InChildFolder);
+		}
 	}
 #endif
 
diff --git a/Engine/Source/Runtime/MovieScene/Public/MovieSceneFolder.h b/Engine/Source/Runtime/MovieScene/Public/MovieSceneFolder.h
--- a/Engine/Source/Runtime/MovieScene/Public/MovieSceneFolder.h
+++ b/Engine/Source/Runtime/MovieScene/Public/MovieSceneFolder.h
@@ -41,7 +41,7 @@ class MOVIESCENE_API UMovieSceneFolder : public UObject
 	/** Gets the folders contained by this folder. */
 	TArrayView<UMovieSceneFolder* const> GetChildFolders() const;
 
-	/** Adds a child folder to this folder. Automatically calls Modify() on the folder object. */
+	/** Adds a child folder to this folder. Removes the folder from any other folders (including root folders). Automatically calls Modify() on the folder object. */
 	void AddChildFolder( UMovieSceneFolder* InChildFolder );
 
 	/** Removes a child folder from this folder. Automatically calls Modify() on the folder object. */
