Sequencer: Add GetSortingOrder/SetSortingOrder for bindings

diff --git a/Engine/Plugins/MovieScene/SequencerScripting/Source/SequencerScripting/Private/ExtensionLibraries/MovieSceneBindingExtensions.cpp b/Engine/Plugins/MovieScene/SequencerScripting/Source/SequencerScripting/Private/ExtensionLibraries/MovieSceneBindingExtensions.cpp
--- a/Engine/Plugins/MovieScene/SequencerScripting/Source/SequencerScripting/Private/ExtensionLibraries/MovieSceneBindingExtensions.cpp
+++ b/Engine/Plugins/MovieScene/SequencerScripting/Source/SequencerScripting/Private/ExtensionLibraries/MovieSceneBindingExtensions.cpp
@@ -90,6 +90,41 @@ void UMovieSceneBindingExtensions::SetName(const FMovieSceneBindingProxy& InBind
 	}
 }
 
+int32 UMovieSceneBindingExtensions::GetSortingOrder(const FMovieSceneBindingProxy& InBinding)
+{
+#if WITH_EDITORONLY_DATA
+	UMovieScene* MovieScene = InBinding.GetMovieScene();
+	if (MovieScene)
+	{
+		const FMovieSceneBinding* Binding = Algo::FindBy(MovieScene->GetBindings(), InBinding.BindingID, &FMovieSceneBinding::GetObjectGuid);
+		if (Binding)
+		{
+			return Binding->GetSortingOrder();
+		}
+	}
+#endif
+
+	FFrame::KismetExecutionMessage(TEXT("Cannot find requested binding"), ELogVerbosity::Error);
+	return 0;
+}
+
+void UMovieSceneBindingExtensions::SetSortingOrder(const FMovieSceneBindingProxy& InBinding, int32 SortingOrder)
+{
+#if WITH_EDITORONLY_DATA
+	UMovieScene* MovieScene = InBinding.GetMovieScene();
+	if (MovieScene)
+	{
+		FMovieSceneBinding* Binding = MovieScene->FindBinding(InBinding.BindingID);
+		if (Binding)
+		{
+			Binding->SetSortingOrder(SortingOrder);
+			return;
+		}
+	}
+#endif
+	FFrame::KismetExecutionMessage(TEXT("Cannot find requested binding"), ELogVerbosity::Error);
+}
+
 TArray<UMovieSceneTrack*> UMovieSceneBindingExtensions::GetTracks(const FMovieSceneBindingProxy& InBinding)
 {
 	UMovieScene* MovieScene = InBinding.GetMovieScene();
diff --git a/Engine/Plugins/MovieScene/SequencerScripting/Source/SequencerScripting/Public/ExtensionLibraries/MovieSceneBindingExtensions.h b/Engine/Plugins/MovieScene/SequencerScripting/Source/SequencerScripting/Public/ExtensionLibraries/MovieSceneBindingExtensions.h
--- a/Engine/Plugins/MovieScene/SequencerScripting/Source/SequencerScripting/Public/ExtensionLibraries/MovieSceneBindingExtensions.h
+++ b/Engine/Plugins/MovieScene/SequencerScripting/Source/SequencerScripting/Public/ExtensionLibraries/MovieSceneBindingExtensions.h
@@ -71,6 +71,24 @@ public:
 	UFUNCTION(BlueprintCallable, Category = "Sequencer|Sequence", meta=(ScriptMethod))
 	static void SetName(const FMovieSceneBindingProxy& InBinding, const FString& InName);
 
+	/**
+	 * Get the sorting order for this binding
+	 *
+	 * @param InBinding        The binding to get the sorting order from
+	 * @return The sorting order of the requested binding
+	 */
+	UFUNCTION(BlueprintCallable, Category = "Sequencer|Sequence", meta = (ScriptMethod, DevelopmentOnly))
+	static int32 GetSortingOrder(const FMovieSceneBindingProxy& InBinding);
+
+	/**
+	 * Set the sorting order for this binding
+	 *
+	 * @param InBinding        The binding to get the sorting order from
+	 * @param SortingOrder The sorting order to set
+	 */
+	UFUNCTION(BlueprintCallable, Category = "Sequencer|Sequence", meta = (ScriptMethod, DevelopmentOnly))
+	static void SetSortingOrder(const FMovieSceneBindingProxy& InBinding, int32 SortingOrder);
+
 	/**
 	 * Get all the tracks stored within this binding
 	 *
