Sequencer: Fix crash on fbx export when RootSequence is null

diff --git a/Engine/Plugins/MovieScene/SequencerScripting/Source/SequencerScriptingEditor/Private/SequencerTools.cpp b/Engine/Plugins/MovieScene/SequencerScripting/Source/SequencerScriptingEditor/Private/SequencerTools.cpp
--- a/Engine/Plugins/MovieScene/SequencerScripting/Source/SequencerScriptingEditor/Private/SequencerTools.cpp
+++ b/Engine/Plugins/MovieScene/SequencerScripting/Source/SequencerScriptingEditor/Private/SequencerTools.cpp
@@ -206,6 +206,16 @@ bool ExportFBXInternal(const FSequencerExportFBXParams& InParams, UMovieSceneSeq
 	UFbxExportOption* OverrideOptions = InParams.OverrideOptions;
 	FString FBXFileName = InParams.FBXFileName;
 
+	if (!RootSequence)
+	{
+		RootSequence = Sequence;
+	}
+
+	if (!RootSequence)
+	{
+		return false;
+	}
+
 	UnFbx::FFbxExporter* Exporter = UnFbx::FFbxExporter::GetInstance();
 	//Show the fbx export dialog options
 	Exporter->SetExportOptionsOverride(OverrideOptions);
@@ -272,7 +282,20 @@ bool USequencerToolsFunctionLibrary::ExportLevelSequenceFBX(const FSequencerExpo
 	ALevelSequenceActor* OutActor;
 	FMovieSceneSequencePlaybackSettings Settings;
 	FLevelSequenceCameraSettings CameraSettings;
-	ULevelSequencePlayer* Player = ULevelSequencePlayer::CreateLevelSequencePlayer(InParams.World, InParams.RootSequence, Settings, OutActor);
+
+	ULevelSequence* RootSequence = InParams.RootSequence;
+	if (!RootSequence)
+	{
+		RootSequence = InParams.Sequence;
+	}
+
+	if (!RootSequence)
+	{
+		FFrame::KismetExecutionMessage(TEXT("Cannot export level sequence. Sequence is invalid."), ELogVerbosity::Error);
+		return false;
+	}
+
+	ULevelSequencePlayer* Player = ULevelSequencePlayer::CreateLevelSequencePlayer(InParams.World, RootSequence, Settings, OutActor);
 
 	bool bSuccess = ExportFBXInternal(InParams, Player);
 	InParams.World->DestroyActor(OutActor);
