[Engine/Plugins]

diff --git a/Engine/Plugins/Experimental/PCG/Source/PCG/Public/MatchAndSet/PCGMatchAndSetWeightedByCategory.h b/Engine/Plugins/Experimental/PCG/Source/PCG/Public/MatchAndSet/PCGMatchAndSetWeightedByCategory.h
--- a/Engine/Plugins/Experimental/PCG/Source/PCG/Public/MatchAndSet/PCGMatchAndSetWeightedByCategory.h
+++ b/Engine/Plugins/Experimental/PCG/Source/PCG/Public/MatchAndSet/PCGMatchAndSetWeightedByCategory.h
@@ -3,11 +3,12 @@
 #pragma once
 
 #include "MatchAndSet/PCGMatchAndSetBase.h"
-#include "MatchAndSet/PCGMatchAndSetWeighted.h"
-#include "Metadata/PCGMetadataTypesConstantStruct.h"
 
 #include "PCGMatchAndSetWeightedByCategory.generated.h"
 
+struct FPCGMatchAndSetWeightedEntry;
+struct FPropertyChangedEvent;
+
 USTRUCT(BlueprintType)
 struct PCG_API FPCGMatchAndSetWeightedByCategoryEntryList
 {
@@ -74,4 +75,8 @@ public:
 	/** Controls whether the output data should mutate its seed - prevents issues when doing multiple random processes in a row */
 	UPROPERTY(BlueprintReadWrite, EditAnywhere, Category = Settings)
 	bool bShouldMutateSeed = true;
-};
\ No newline at end of file
+};
+
+#if UE_ENABLE_INCLUDE_ORDER_DEPRECATED_IN_5_2
+#include "MatchAndSet/PCGMatchAndSetWeighted.h"
+#endif
diff --git a/Engine/Plugins/Experimental/PCG/Source/PCG/Public/MeshSelectors/PCGMeshSelectorByAttribute.h b/Engine/Plugins/Experimental/PCG/Source/PCG/Public/MeshSelectors/PCGMeshSelectorByAttribute.h
--- a/Engine/Plugins/Experimental/PCG/Source/PCG/Public/MeshSelectors/PCGMeshSelectorByAttribute.h
+++ b/Engine/Plugins/Experimental/PCG/Source/PCG/Public/MeshSelectors/PCGMeshSelectorByAttribute.h
@@ -4,7 +4,6 @@
 
 #include "PCGMeshSelectorBase.h"
 
-#include "Engine/CollisionProfile.h"
 
 #include "PCGMeshSelectorByAttribute.generated.h"
 
diff --git a/Engine/Plugins/Experimental/PCG/Source/PCG/Public/MeshSelectors/PCGMeshSelectorWeighted.h b/Engine/Plugins/Experimental/PCG/Source/PCG/Public/MeshSelectors/PCGMeshSelectorWeighted.h
--- a/Engine/Plugins/Experimental/PCG/Source/PCG/Public/MeshSelectors/PCGMeshSelectorWeighted.h
+++ b/Engine/Plugins/Experimental/PCG/Source/PCG/Public/MeshSelectors/PCGMeshSelectorWeighted.h
@@ -4,7 +4,6 @@
 
 #include "PCGMeshSelectorBase.h"
 
-#include "Engine/CollisionProfile.h"
 
 #include "PCGMeshSelectorWeighted.generated.h"
 
diff --git a/Engine/Plugins/Experimental/PCG/Source/PCGEditor/Private/SPCGEditorGraphProfilingView.cpp b/Engine/Plugins/Experimental/PCG/Source/PCGEditor/Private/SPCGEditorGraphProfilingView.cpp
--- a/Engine/Plugins/Experimental/PCG/Source/PCGEditor/Private/SPCGEditorGraphProfilingView.cpp
+++ b/Engine/Plugins/Experimental/PCG/Source/PCGEditor/Private/SPCGEditorGraphProfilingView.cpp
@@ -7,9 +7,6 @@
 #include "PCGEditor.h"
 #include "PCGEditorGraph.h"
 #include "PCGEditorGraphNode.h"
-#include "PCGGraph.h"
-#include "PCGSubgraph.h"
-#include "PCGSubsystem.h"
 
 #include "Widgets/Input/SButton.h"
 #include "Widgets/Input/SCheckBox.h"
diff --git a/Engine/Plugins/Experimental/PlanarCutPlugin/Source/PlanarCut/Public/FractureAutoUV.h b/Engine/Plugins/Experimental/PlanarCutPlugin/Source/PlanarCut/Public/FractureAutoUV.h
--- a/Engine/Plugins/Experimental/PlanarCutPlugin/Source/PlanarCut/Public/FractureAutoUV.h
+++ b/Engine/Plugins/Experimental/PlanarCutPlugin/Source/PlanarCut/Public/FractureAutoUV.h
@@ -1,10 +1,12 @@
 // Copyright Epic Games, Inc. All Rights Reserved.
 #pragma once
-#include "CoreMinimal.h"
 
-#include "GeometryCollection/GeometryCollection.h"
 
-#include "Image/ImageBuilder.h"
+#include "Math/Vector4.h"
+
+class FGeometryCollection;
+namespace UE::Geometry { struct FIndex4i; }
+namespace UE::Geometry { template <typename PixelType> class TImageBuilder; }
 
 class FProgressCancel;
 
@@ -121,4 +123,10 @@ bool PLANARCUT_API TextureInternalSurfaces(
 	FProgressCancel* Progress = nullptr
 );
 
-}} // namespace UE::PlanarCut
\ No newline at end of file
+}} // namespace UE::PlanarCut
+
+#if UE_ENABLE_INCLUDE_ORDER_DEPRECATED_IN_5_2
+#include "CoreMinimal.h"
+#include "GeometryCollection/GeometryCollection.h"
+#include "Image/ImageBuilder.h"
+#endif
diff --git a/Engine/Plugins/Experimental/PlanarCutPlugin/Source/PlanarCut/Public/PlanarCut.h b/Engine/Plugins/Experimental/PlanarCutPlugin/Source/PlanarCut/Public/PlanarCut.h
--- a/Engine/Plugins/Experimental/PlanarCutPlugin/Source/PlanarCut/Public/PlanarCut.h
+++ b/Engine/Plugins/Experimental/PlanarCutPlugin/Source/PlanarCut/Public/PlanarCut.h
@@ -1,15 +1,14 @@
 // Copyright Epic Games, Inc. All Rights Reserved.
 #pragma once
-#include "CoreMinimal.h"
 
-#include "Math/Plane.h"
 
-#include "Voronoi/Voronoi.h"
 #include "GeometryCollection/GeometryCollection.h"
-#include "GeometryCollection/Facades/CollectionMeshFacade.h"
-#include "MeshDescription.h"
 #include "DynamicMesh/DynamicMesh3.h"
 
+class FVoronoiDiagram;
+namespace GeometryCollection::Facades { class FCollectionMeshFacade; }
+struct FMeshDescription;
+
 class FProgressCancel;
 
 // Offsets to randomize where Perlin noise is sampled for each dimension of a noise vector
@@ -509,3 +508,10 @@ void PLANARCUT_API ConvertToMeshDescription(
 	const TArrayView<const int32>& TransformIndices
 );
 
+
+#if UE_ENABLE_INCLUDE_ORDER_DEPRECATED_IN_5_2
+#include "CoreMinimal.h"
+#include "GeometryCollection/Facades/CollectionMeshFacade.h"
+#include "MeshDescription.h"
+#include "Voronoi/Voronoi.h"
+#endif
diff --git a/Engine/Plugins/Experimental/StructUtils/Source/StructUtils/Public/InstancedStructContainer.h b/Engine/Plugins/Experimental/StructUtils/Source/StructUtils/Public/InstancedStructContainer.h
--- a/Engine/Plugins/Experimental/StructUtils/Source/StructUtils/Public/InstancedStructContainer.h
+++ b/Engine/Plugins/Experimental/StructUtils/Source/StructUtils/Public/InstancedStructContainer.h
@@ -2,11 +2,11 @@
 
 #pragma once
 
-#include "InstancedStruct.h"
 #include "StructView.h"
-#include "Templates/ChooseClass.h"
 #include "InstancedStructContainer.generated.h"
 
+struct FInstancedStruct;
+
 /**
  * Array of heterogeneous structs. Can be used as a property, supports serialization,
  * but does not have type customizations (no editing in the UI).
diff --git a/Engine/Plugins/Experimental/StructUtils/Source/StructUtils/Public/PropertyBag.h b/Engine/Plugins/Experimental/StructUtils/Source/StructUtils/Public/PropertyBag.h
--- a/Engine/Plugins/Experimental/StructUtils/Source/StructUtils/Public/PropertyBag.h
+++ b/Engine/Plugins/Experimental/StructUtils/Source/StructUtils/Public/PropertyBag.h
@@ -1,12 +1,10 @@
-﻿// Copyright Epic Games, Inc. All Rights Reserved.
+// Copyright Epic Games, Inc. All Rights Reserved.
 
 #pragma once
 
-#include "InstancedStruct.h"
+#include "Misc/TVariantMeta.h"
 #include "StructView.h"
 #include "Templates/ValueOrError.h"
-#include "UObject/UnrealType.h"
-#include <atomic>
 #include "PropertyBag.generated.h"
 
 /** Property bag property type, loosely based on BluePrint pin types. */
diff --git a/Engine/Plugins/Experimental/StructUtils/Source/StructUtils/Public/StructArrayView.h b/Engine/Plugins/Experimental/StructUtils/Source/StructUtils/Public/StructArrayView.h
--- a/Engine/Plugins/Experimental/StructUtils/Source/StructUtils/Public/StructArrayView.h
+++ b/Engine/Plugins/Experimental/StructUtils/Source/StructUtils/Public/StructArrayView.h
@@ -2,9 +2,7 @@
 
 #pragma once
 
-#include "StructUtils.h"
-#include "InstancedStruct.h"
-#include "SharedStruct.h"
+#include "UObject/Class.h"
 
 ///////////////////////////////////////////////////////////////// FStructArrayView /////////////////////////////////////////////////////////////////
 
@@ -103,3 +101,9 @@ private:
 	int32 NumElements = 0;
 	UScriptStruct& FragmentType;
 };
+
+#if UE_ENABLE_INCLUDE_ORDER_DEPRECATED_IN_5_2
+#include "InstancedStruct.h"
+#include "SharedStruct.h"
+#include "StructUtils.h"
+#endif
diff --git a/Engine/Plugins/Experimental/StructUtils/Source/StructUtils/Public/StructTypeBitSet.h b/Engine/Plugins/Experimental/StructUtils/Source/StructUtils/Public/StructTypeBitSet.h
--- a/Engine/Plugins/Experimental/StructUtils/Source/StructUtils/Public/StructTypeBitSet.h
+++ b/Engine/Plugins/Experimental/StructUtils/Source/StructUtils/Public/StructTypeBitSet.h
@@ -2,9 +2,6 @@
 
 #pragma once
 
-#include "UObject/Class.h"
-#include "Containers/BitArray.h"
-#include "InstancedStruct.h"
 #include "StructUtilsTypes.h"
 
 
diff --git a/Engine/Plugins/Experimental/StructUtils/Source/StructUtils/Public/StructUtilsTypes.h b/Engine/Plugins/Experimental/StructUtils/Source/StructUtils/Public/StructUtilsTypes.h
--- a/Engine/Plugins/Experimental/StructUtils/Source/StructUtils/Public/StructUtilsTypes.h
+++ b/Engine/Plugins/Experimental/StructUtils/Source/StructUtils/Public/StructUtilsTypes.h
@@ -2,9 +2,7 @@
 
 #pragma once
 
-#include "UObject/Class.h"
 #include "StructView.h"
-#include "Templates/UnrealTypeTraits.h"
 
 #ifndef WITH_STRUCTUTILS_DEBUG
 #define WITH_STRUCTUTILS_DEBUG (!(UE_BUILD_SHIPPING || UE_BUILD_SHIPPING_WITH_EDITOR || UE_BUILD_TEST) && 1)
@@ -68,4 +66,4 @@ struct FStructTypeSortOperator
 		}
 		return FScriptStructSortOperator()(*AScriptStruct, *BScriptStruct);
 	}
-};
\ No newline at end of file
+};
diff --git a/Engine/Plugins/Experimental/StructUtils/Source/StructUtils/Public/StructView.h b/Engine/Plugins/Experimental/StructUtils/Source/StructUtils/Public/StructView.h
--- a/Engine/Plugins/Experimental/StructUtils/Source/StructUtils/Public/StructView.h
+++ b/Engine/Plugins/Experimental/StructUtils/Source/StructUtils/Public/StructView.h
@@ -2,7 +2,6 @@
 
 #pragma once
 
-#include "StructUtils.h"
 #include "InstancedStruct.h"
 #include "SharedStruct.h"
 
@@ -215,4 +214,4 @@ public:
 		}
 		return nullptr;
 	}
-};
\ No newline at end of file
+};
diff --git a/Engine/Plugins/Experimental/StructUtils/Source/StructUtilsEditor/Public/InstancedStructDetails.h b/Engine/Plugins/Experimental/StructUtils/Source/StructUtilsEditor/Public/InstancedStructDetails.h
--- a/Engine/Plugins/Experimental/StructUtils/Source/StructUtilsEditor/Public/InstancedStructDetails.h
+++ b/Engine/Plugins/Experimental/StructUtils/Source/StructUtilsEditor/Public/InstancedStructDetails.h
@@ -2,8 +2,6 @@
 
 #pragma once
 
-#include "CoreMinimal.h"
-#include "UnrealClient.h"
 #include "IPropertyTypeCustomization.h"
 #include "IDetailCustomNodeBuilder.h"
 #include "EditorUndoClient.h"
@@ -97,3 +95,8 @@ private:
 	/** True if we're currently handling a StructValuePostChange */
 	bool bIsHandlingStructValuePostChange = false;
 };
+
+#if UE_ENABLE_INCLUDE_ORDER_DEPRECATED_IN_5_2
+#include "CoreMinimal.h"
+#include "UnrealClient.h"
+#endif
diff --git a/Engine/Plugins/Experimental/StructUtils/Source/StructUtilsEditor/Public/PropertyBagDetails.h b/Engine/Plugins/Experimental/StructUtils/Source/StructUtilsEditor/Public/PropertyBagDetails.h
--- a/Engine/Plugins/Experimental/StructUtils/Source/StructUtilsEditor/Public/PropertyBagDetails.h
+++ b/Engine/Plugins/Experimental/StructUtils/Source/StructUtilsEditor/Public/PropertyBagDetails.h
@@ -2,11 +2,14 @@
 
 #pragma once
 
-#include "IPropertyTypeCustomization.h"
 #include "InstancedStructDetails.h"
 #include "EdGraphSchema_K2.h"
 #include "PropertyBagDetails.generated.h"
 
+class FReply;
+enum class EPinContainerType : uint8;
+struct FEdGraphSchemaAction;
+
 class IPropertyHandle;
 class IPropertyUtilities;
 class IDetailPropertyRow;
diff --git a/Engine/Plugins/Experimental/StructUtils/Source/StructUtilsEditor/Public/StructUtilsEditorModule.h b/Engine/Plugins/Experimental/StructUtils/Source/StructUtilsEditor/Public/StructUtilsEditorModule.h
--- a/Engine/Plugins/Experimental/StructUtils/Source/StructUtilsEditor/Public/StructUtilsEditorModule.h
+++ b/Engine/Plugins/Experimental/StructUtils/Source/StructUtilsEditor/Public/StructUtilsEditorModule.h
@@ -2,11 +2,7 @@
 
 #pragma  once
 
-#include "CoreMinimal.h"
-#include "Modules/ModuleManager.h"
-#include "AssetTypeCategories.h"
-#include "Toolkits/IToolkitHost.h"
-#include "Toolkits/AssetEditorToolkit.h"
+#include "Modules/ModuleInterface.h"
 
 
 class IStructUtilsEditor;
@@ -23,3 +19,11 @@ public:
 	virtual void StartupModule() override;
 	virtual void ShutdownModule() override;
 };
+
+#if UE_ENABLE_INCLUDE_ORDER_DEPRECATED_IN_5_2
+#include "AssetTypeCategories.h"
+#include "CoreMinimal.h"
+#include "Modules/ModuleManager.h"
+#include "Toolkits/AssetEditorToolkit.h"
+#include "Toolkits/IToolkitHost.h"
+#endif
diff --git a/Engine/Plugins/Experimental/UIFramework/Source/Private/UIFModule.cpp b/Engine/Plugins/Experimental/UIFramework/Source/Private/UIFModule.cpp
--- a/Engine/Plugins/Experimental/UIFramework/Source/Private/UIFModule.cpp
+++ b/Engine/Plugins/Experimental/UIFramework/Source/Private/UIFModule.cpp
@@ -6,7 +6,6 @@
 #include "UIFLocalSettings.h"
 #include "UIFPresenter.h"
 #include "UIFWidget.h"
-#include "UObject/Package.h"
 
 namespace UE::UIFramework::Private
 {
@@ -167,4 +166,4 @@ void FUIFrameworkModule::StartupModule()
 	GetDefault<UUIFrameworkLocalSettings>()->LoadResources();
 }
 
-IMPLEMENT_MODULE(FUIFrameworkModule, UIFramework)
\ No newline at end of file
+IMPLEMENT_MODULE(FUIFrameworkModule, UIFramework)
diff --git a/Engine/Plugins/Experimental/UIFramework/Source/Private/Widgets/UIFColorBlock.cpp b/Engine/Plugins/Experimental/UIFramework/Source/Private/Widgets/UIFColorBlock.cpp
--- a/Engine/Plugins/Experimental/UIFramework/Source/Private/Widgets/UIFColorBlock.cpp
+++ b/Engine/Plugins/Experimental/UIFramework/Source/Private/Widgets/UIFColorBlock.cpp
@@ -3,11 +3,6 @@
 #include "Widgets/UIFColorBlock.h"
 
 #include "Components/Image.h"
-#include "Engine/AssetManager.h"
-#include "Engine/StreamableManager.h"
-#include "Engine/Texture2D.h"
-#include "Materials/Material.h"
-#include "Materials/MaterialInterface.h"
 
 #include "Net/Core/PushModel/PushModel.h"
 #include "Net/UnrealNetwork.h"
diff --git a/Engine/Plugins/Experimental/UIFramework/Source/Private/Widgets/UIFImageBlock.cpp b/Engine/Plugins/Experimental/UIFramework/Source/Private/Widgets/UIFImageBlock.cpp
--- a/Engine/Plugins/Experimental/UIFramework/Source/Private/Widgets/UIFImageBlock.cpp
+++ b/Engine/Plugins/Experimental/UIFramework/Source/Private/Widgets/UIFImageBlock.cpp
@@ -4,10 +4,8 @@
 
 #include "Components/Image.h"
 #include "Engine/AssetManager.h"
-#include "Engine/StreamableManager.h"
 #include "Engine/Texture2D.h"
 #include "Materials/Material.h"
-#include "Materials/MaterialInterface.h"
 #include "UIFLocalSettings.h"
 
 #include "Net/Core/PushModel/PushModel.h"
diff --git a/Engine/Plugins/Experimental/UIFramework/Source/Public/Types/UIFWidgetTreeOwner.h b/Engine/Plugins/Experimental/UIFramework/Source/Public/Types/UIFWidgetTreeOwner.h
--- a/Engine/Plugins/Experimental/UIFramework/Source/Public/Types/UIFWidgetTreeOwner.h
+++ b/Engine/Plugins/Experimental/UIFramework/Source/Public/Types/UIFWidgetTreeOwner.h
@@ -2,7 +2,7 @@
 
 #pragma once
 
-#include "HAL/Platform.h"
+#include "HAL/Platform.h" // IWYU pragma: keep
 
 struct FUIFrameworkWidgetOwner;
 
@@ -37,5 +37,6 @@ public:
 };
 
 #if UE_ENABLE_INCLUDE_ORDER_DEPRECATED_IN_5_2
+#include "HAL/Platform.h"
 #include "Types/UIFWidgetOwner.h"
 #endif
diff --git a/Engine/Plugins/Experimental/UIFramework/Source/Public/Widgets/UIFColorBlock.h b/Engine/Plugins/Experimental/UIFramework/Source/Public/Widgets/UIFColorBlock.h
--- a/Engine/Plugins/Experimental/UIFramework/Source/Public/Widgets/UIFColorBlock.h
+++ b/Engine/Plugins/Experimental/UIFramework/Source/Public/Widgets/UIFColorBlock.h
@@ -2,7 +2,6 @@
 
 #pragma once
 
-#include "Styling/SlateBrush.h"
 #include "UIFWidget.h"
 
 #include "UIFColorBlock.generated.h"
diff --git a/Engine/Plugins/Experimental/UIFramework/Source/Public/Widgets/UIFImageBlock.h b/Engine/Plugins/Experimental/UIFramework/Source/Public/Widgets/UIFImageBlock.h
--- a/Engine/Plugins/Experimental/UIFramework/Source/Public/Widgets/UIFImageBlock.h
+++ b/Engine/Plugins/Experimental/UIFramework/Source/Public/Widgets/UIFImageBlock.h
@@ -2,7 +2,6 @@
 
 #pragma once
 
-#include "Styling/SlateBrush.h"
 #include "UIFWidget.h"
 
 #include "UIFImageBlock.generated.h"
diff --git a/Engine/Plugins/Experimental/UIFramework/Source/Public/Widgets/UIFOverlay.h b/Engine/Plugins/Experimental/UIFramework/Source/Public/Widgets/UIFOverlay.h
--- a/Engine/Plugins/Experimental/UIFramework/Source/Public/Widgets/UIFOverlay.h
+++ b/Engine/Plugins/Experimental/UIFramework/Source/Public/Widgets/UIFOverlay.h
@@ -2,15 +2,13 @@
 
 #pragma once
 
-#include "Components/SlateWrapperTypes.h"
-#include "Layout/Margin.h"
-#include "Types/SlateEnums.h"
 #include "Types/UIFSlotBase.h"
 #include "UIFWidget.h"
-#include "Widgets/Layout/Anchors.h"
 
 #include "UIFOverlay.generated.h"
 
+struct FUIFrameworkWidgetId;
+
 class UUIFrameworkOverlay;
 struct FUIFrameworkOverlaySlotList;
 
@@ -116,3 +114,7 @@ private:
 	UPROPERTY(Replicated)
 	FUIFrameworkOverlaySlotList ReplicatedSlotList;
 };
+
+#if UE_ENABLE_INCLUDE_ORDER_DEPRECATED_IN_5_2
+#include "Widgets/Layout/Anchors.h"
+#endif
diff --git a/Engine/Plugins/Experimental/VirtualHeightfieldMesh/Source/VirtualHeightfieldMesh/Public/HeightfieldMinMaxRender.h b/Engine/Plugins/Experimental/VirtualHeightfieldMesh/Source/VirtualHeightfieldMesh/Public/HeightfieldMinMaxRender.h
--- a/Engine/Plugins/Experimental/VirtualHeightfieldMesh/Source/VirtualHeightfieldMesh/Public/HeightfieldMinMaxRender.h
+++ b/Engine/Plugins/Experimental/VirtualHeightfieldMesh/Source/VirtualHeightfieldMesh/Public/HeightfieldMinMaxRender.h
@@ -2,7 +2,7 @@
 
 #pragma once
 
-#include "CoreMinimal.h"
+#include "Math/MathFwd.h"
 
 class FRDGBuilder;
 class FRDGTexture;
@@ -19,3 +19,7 @@ namespace VirtualHeightfieldMesh
 	/** Utility function to generate all additional mips from mip0 for a MinMax height texture already packed in RGBA8. */
 	VIRTUALHEIGHTFIELDMESH_API void GenerateMinMaxTextureMips(FRDGBuilder& GraphBuilder, FRDGTexture* Texture, FIntPoint SrcSize, int32 NumMips);
 };
+
+#if UE_ENABLE_INCLUDE_ORDER_DEPRECATED_IN_5_2
+#include "CoreMinimal.h"
+#endif
diff --git a/Engine/Plugins/Experimental/VirtualHeightfieldMesh/Source/VirtualHeightfieldMesh/Public/HeightfieldMinMaxTexture.h b/Engine/Plugins/Experimental/VirtualHeightfieldMesh/Source/VirtualHeightfieldMesh/Public/HeightfieldMinMaxTexture.h
--- a/Engine/Plugins/Experimental/VirtualHeightfieldMesh/Source/VirtualHeightfieldMesh/Public/HeightfieldMinMaxTexture.h
+++ b/Engine/Plugins/Experimental/VirtualHeightfieldMesh/Source/VirtualHeightfieldMesh/Public/HeightfieldMinMaxTexture.h
@@ -2,10 +2,10 @@
 
 #pragma once
 
-#include "CoreMinimal.h"
-#include "Engine/Texture.h"
 #include "HeightfieldMinMaxTexture.generated.h"
 
+struct FPropertyChangedEvent;
+
 #if WITH_EDITOR
 
 /** Description object used to build the contents of a UHeightfieldMinMaxTexture. */
@@ -77,3 +77,8 @@ protected:
 	//~ End UObject Interface
 #endif
 };
+
+#if UE_ENABLE_INCLUDE_ORDER_DEPRECATED_IN_5_2
+#include "CoreMinimal.h"
+#include "Engine/Texture.h"
+#endif
diff --git a/Engine/Plugins/Experimental/VirtualHeightfieldMesh/Source/VirtualHeightfieldMesh/Public/VirtualHeightfieldMeshComponent.h b/Engine/Plugins/Experimental/VirtualHeightfieldMesh/Source/VirtualHeightfieldMesh/Public/VirtualHeightfieldMeshComponent.h
--- a/Engine/Plugins/Experimental/VirtualHeightfieldMesh/Source/VirtualHeightfieldMesh/Public/VirtualHeightfieldMeshComponent.h
+++ b/Engine/Plugins/Experimental/VirtualHeightfieldMesh/Source/VirtualHeightfieldMesh/Public/VirtualHeightfieldMeshComponent.h
@@ -2,8 +2,8 @@
 
 #pragma once
 
-#include "CoreMinimal.h"
 #include "Components/PrimitiveComponent.h"
+#include "Physics/PhysicsInterfaceTypes.h"
 #include "VirtualHeightfieldMeshComponent.generated.h"
 
 class ARuntimeVirtualTextureVolume;
@@ -149,3 +149,7 @@ protected:
 	virtual bool SupportsWorldPositionOffsetVelocity() const override { return bWorldPositionOffsetVelocity; }
 	//~ End UPrimitiveComponent Interface
 };
+
+#if UE_ENABLE_INCLUDE_ORDER_DEPRECATED_IN_5_2
+#include "CoreMinimal.h"
+#endif
diff --git a/Engine/Plugins/Experimental/VirtualHeightfieldMesh/Source/VirtualHeightfieldMesh/Public/VirtualHeightfieldMeshModule.h b/Engine/Plugins/Experimental/VirtualHeightfieldMesh/Source/VirtualHeightfieldMesh/Public/VirtualHeightfieldMeshModule.h
--- a/Engine/Plugins/Experimental/VirtualHeightfieldMesh/Source/VirtualHeightfieldMesh/Public/VirtualHeightfieldMeshModule.h
+++ b/Engine/Plugins/Experimental/VirtualHeightfieldMesh/Source/VirtualHeightfieldMesh/Public/VirtualHeightfieldMeshModule.h
@@ -3,7 +3,6 @@
 #pragma once
 
 #include "Modules/ModuleInterface.h"
-#include "Modules/ModuleManager.h"
 
 /** Implements the VirtualHeightfieldMesh module  */
 class FVirtualHeightfieldMeshModule : public IModuleInterface
@@ -14,3 +13,7 @@ public:
 	virtual void ShutdownModule() override;
 	//~ End IModuleInterface Interface.
 };
+
+#if UE_ENABLE_INCLUDE_ORDER_DEPRECATED_IN_5_2
+#include "Modules/ModuleManager.h"
+#endif
diff --git a/Engine/Plugins/Experimental/VirtualHeightfieldMesh/Source/VirtualHeightfieldMeshEditor/Public/VirtualHeightfieldMeshEditorModule.h b/Engine/Plugins/Experimental/VirtualHeightfieldMesh/Source/VirtualHeightfieldMeshEditor/Public/VirtualHeightfieldMeshEditorModule.h
--- a/Engine/Plugins/Experimental/VirtualHeightfieldMesh/Source/VirtualHeightfieldMeshEditor/Public/VirtualHeightfieldMeshEditorModule.h
+++ b/Engine/Plugins/Experimental/VirtualHeightfieldMesh/Source/VirtualHeightfieldMeshEditor/Public/VirtualHeightfieldMeshEditorModule.h
@@ -3,7 +3,6 @@
 #pragma once
 
 #include "Modules/ModuleInterface.h"
-#include "Modules/ModuleManager.h"
 
 /** Module interface for VirtualHeightfieldMesh editor  extensions. */
 class IVirtualHeightfieldMeshEditorModule : public IModuleInterface
@@ -14,3 +13,7 @@ public:
 	/** Build the contents of the streaming low mips. */
 	virtual bool BuildMinMaxHeightTexture(class UVirtualHeightfieldMeshComponent* InComponent) const = 0;
 };
+
+#if UE_ENABLE_INCLUDE_ORDER_DEPRECATED_IN_5_2
+#include "Modules/ModuleManager.h"
+#endif
diff --git a/Engine/Plugins/Experimental/Water/Source/Runtime/Public/WaterInstanceDataBuffer.h b/Engine/Plugins/Experimental/Water/Source/Runtime/Public/WaterInstanceDataBuffer.h
--- a/Engine/Plugins/Experimental/Water/Source/Runtime/Public/WaterInstanceDataBuffer.h
+++ b/Engine/Plugins/Experimental/Water/Source/Runtime/Public/WaterInstanceDataBuffer.h
@@ -2,10 +2,7 @@
 
 #pragma once
 
-#include "Containers/ClosableMpscQueue.h"
 #include "RenderingThread.h"
-#include "RHI.h"
-#include "RHIFwd.h"
 
 template <bool bWithWaterSelectionSupport>
 class TWaterInstanceDataBuffers
@@ -101,5 +98,6 @@ private:
 };
 
 #if UE_ENABLE_INCLUDE_ORDER_DEPRECATED_IN_5_2
+#include "RHI.h"
 #include "CoreMinimal.h"
 #endif
diff --git a/Engine/Plugins/Experimental/Water/Source/Runtime/Public/WaterMeshComponent.h b/Engine/Plugins/Experimental/Water/Source/Runtime/Public/WaterMeshComponent.h
--- a/Engine/Plugins/Experimental/Water/Source/Runtime/Public/WaterMeshComponent.h
+++ b/Engine/Plugins/Experimental/Water/Source/Runtime/Public/WaterMeshComponent.h
@@ -6,6 +6,8 @@
 #include "WaterQuadTree.h"
 #include "WaterMeshComponent.generated.h"
 
+struct FPSOPrecacheParams;
+
 /**
  * Water Mesh Component responsible for generating and rendering a continuous water mesh on top of all the existing water body actors in the world
  * The component contains a quadtree which defines where there are water tiles. A function for traversing the quadtree and outputing a list of instance data for each tile to be rendered from a point of view is included
diff --git a/Engine/Plugins/Importers/AlembicImporter/Source/AlembicLibrary/Public/AbcFile.h b/Engine/Plugins/Importers/AlembicImporter/Source/AlembicLibrary/Public/AbcFile.h
--- a/Engine/Plugins/Importers/AlembicImporter/Source/AlembicLibrary/Public/AbcFile.h
+++ b/Engine/Plugins/Importers/AlembicImporter/Source/AlembicLibrary/Public/AbcFile.h
@@ -2,8 +2,7 @@
 
 #pragma once
 
-#include "CoreMinimal.h"
-#include "Misc/EnumClassFlags.h"
+#include "Math/BoxSphereBounds.h"
 
 #if PLATFORM_WINDOWS
 #include "Windows/WindowsHWrapper.h"
@@ -171,3 +170,8 @@ protected:
 
 	TMap<FString, FString> CustomAttributes;
 };
+
+#if UE_ENABLE_INCLUDE_ORDER_DEPRECATED_IN_5_2
+#include "UObject/NoExportTypes.h"
+#include "CoreMinimal.h"
+#endif
diff --git a/Engine/Plugins/Importers/AlembicImporter/Source/AlembicLibrary/Public/AbcImportLogger.h b/Engine/Plugins/Importers/AlembicImporter/Source/AlembicLibrary/Public/AbcImportLogger.h
--- a/Engine/Plugins/Importers/AlembicImporter/Source/AlembicLibrary/Public/AbcImportLogger.h
+++ b/Engine/Plugins/Importers/AlembicImporter/Source/AlembicLibrary/Public/AbcImportLogger.h
@@ -2,8 +2,9 @@
 
 #pragma once
 
-#include "CoreMinimal.h"
-#include "Logging/TokenizedMessage.h"
+#include "Templates/SharedPointer.h"
+
+class FTokenizedMessage;
 
 class FAbcImportLogger
 {
@@ -22,3 +23,8 @@ private:
 	static FCriticalSection MessageLock;
 };
 
+
+#if UE_ENABLE_INCLUDE_ORDER_DEPRECATED_IN_5_2
+#include "CoreMinimal.h"
+#include "Logging/TokenizedMessage.h"
+#endif
diff --git a/Engine/Plugins/Importers/AlembicImporter/Source/AlembicLibrary/Public/AbcImportSettings.h b/Engine/Plugins/Importers/AlembicImporter/Source/AlembicLibrary/Public/AbcImportSettings.h
--- a/Engine/Plugins/Importers/AlembicImporter/Source/AlembicLibrary/Public/AbcImportSettings.h
+++ b/Engine/Plugins/Importers/AlembicImporter/Source/AlembicLibrary/Public/AbcImportSettings.h
@@ -2,8 +2,6 @@
 
 #pragma once
 
-#include "CoreMinimal.h"
-#include "UObject/Object.h"
 #include "AbcImportSettings.generated.h"
 
 /** Enum that describes type of asset to import */
@@ -345,3 +343,7 @@ class ALEMBICLIBRARY_API UAbcImportSettings : public UObject
 public:
 	virtual void Serialize(class FArchive& Archive) override;
 };
+
+#if UE_ENABLE_INCLUDE_ORDER_DEPRECATED_IN_5_2
+#include "CoreMinimal.h"
+#endif
diff --git a/Engine/Plugins/Importers/AlembicImporter/Source/AlembicLibrary/Public/AbcImporter.h b/Engine/Plugins/Importers/AlembicImporter/Source/AlembicLibrary/Public/AbcImporter.h
--- a/Engine/Plugins/Importers/AlembicImporter/Source/AlembicLibrary/Public/AbcImporter.h
+++ b/Engine/Plugins/Importers/AlembicImporter/Source/AlembicLibrary/Public/AbcImporter.h
@@ -2,11 +2,8 @@
 
 #pragma once
 
-#include "CoreMinimal.h"
-#include "Containers/List.h"
-#include "Animation/MorphTarget.h"
-#include "Animation/AnimSequence.h"
-#include "AbcPolyMesh.h"
+#include "MeshUVChannelInfo.h"
+#include "UObject/Object.h" // IWYU pragma: keep
 
 #if PLATFORM_WINDOWS
 #include "Windows/WindowsHWrapper.h"
@@ -21,6 +18,9 @@ THIRD_PARTY_INCLUDES_END
 #include "Windows/HideWindowsPlatformTypes.h"
 #endif
 
+class IAnimationDataController;
+class UAnimSequence;
+class USkeleton;
 class FSkeletalMeshLODModel;
 class UMaterial;
 class UStaticMesh;
@@ -33,6 +33,8 @@ class UAbcImportSettings;
 class FSkeletalMeshImportData;
 class UAbcAssetImportData;
 
+struct FMorphTargetDelta;
+struct FReferenceSkeleton;
 struct FMeshDescription;
 struct FAbcImportData;
 struct FGeometryCacheMeshData;
@@ -221,3 +223,10 @@ private:
 	/** ABC file representation for currently opened filed */
 	class FAbcFile* AbcFile;
 };
+
+#if UE_ENABLE_INCLUDE_ORDER_DEPRECATED_IN_5_2
+#include "AbcPolyMesh.h"
+#include "Animation/AnimSequence.h"
+#include "Animation/MorphTarget.h"
+#include "CoreMinimal.h"
+#endif
diff --git a/Engine/Plugins/Importers/AlembicImporter/Source/AlembicLibrary/Public/AbcObject.h b/Engine/Plugins/Importers/AlembicImporter/Source/AlembicLibrary/Public/AbcObject.h
--- a/Engine/Plugins/Importers/AlembicImporter/Source/AlembicLibrary/Public/AbcObject.h
+++ b/Engine/Plugins/Importers/AlembicImporter/Source/AlembicLibrary/Public/AbcObject.h
@@ -2,8 +2,7 @@
 
 #pragma once
 
-#include "CoreMinimal.h"
-#include "HAL/ThreadSafeBool.h"
+#include "Math/Matrix.h"
 
 #if PLATFORM_WINDOWS
 #include "Windows/WindowsHWrapper.h"
@@ -72,3 +71,8 @@ protected:
 	int32 ResidentSampleIndices[MaxNumberOfResidentSamples];
 	bool InUseSamples[MaxNumberOfResidentSamples];
 };
+
+#if UE_ENABLE_INCLUDE_ORDER_DEPRECATED_IN_5_2
+#include "CoreMinimal.h"
+#include "HAL/ThreadSafeBool.h"
+#endif
diff --git a/Engine/Plugins/Importers/AlembicImporter/Source/AlembicLibrary/Public/AbcPolyMesh.h b/Engine/Plugins/Importers/AlembicImporter/Source/AlembicLibrary/Public/AbcPolyMesh.h
--- a/Engine/Plugins/Importers/AlembicImporter/Source/AlembicLibrary/Public/AbcPolyMesh.h
+++ b/Engine/Plugins/Importers/AlembicImporter/Source/AlembicLibrary/Public/AbcPolyMesh.h
@@ -3,7 +3,8 @@
 #pragma once
 
 #include "AbcObject.h"
-#include "Components.h"
+#include "Math/BoxSphereBounds.h"
+#include "MeshUVChannelInfo.h"
 
 enum class ESampleReadFlags : uint8;
 
@@ -138,3 +139,7 @@ protected:
 	/** Whether or not the mesh is visible in the first Frame */
 	bool bFirstFrameVisibility;
 };
+
+#if UE_ENABLE_INCLUDE_ORDER_DEPRECATED_IN_5_2
+#include "Components.h"
+#endif
diff --git a/Engine/Plugins/Importers/AlembicImporter/Source/AlembicLibrary/Public/AbcUtilities.h b/Engine/Plugins/Importers/AlembicImporter/Source/AlembicLibrary/Public/AbcUtilities.h
--- a/Engine/Plugins/Importers/AlembicImporter/Source/AlembicLibrary/Public/AbcUtilities.h
+++ b/Engine/Plugins/Importers/AlembicImporter/Source/AlembicLibrary/Public/AbcUtilities.h
@@ -2,7 +2,9 @@
 
 #pragma once
 
-#include "CoreMinimal.h"
+#include "HAL/Platform.h"
+
+class UObject;
 
 class FAbcFile;
 struct FGeometryCacheMeshData;
@@ -17,4 +19,8 @@ public:
 
 	/** Sets up materials from an AbcFile to a GeometryCache and moves them into the given package */
 	static void SetupGeometryCacheMaterials(FAbcFile& AbcFile, class UGeometryCache* GeometryCache, UObject* Package);
-};
\ No newline at end of file
+};
+
+#if UE_ENABLE_INCLUDE_ORDER_DEPRECATED_IN_5_2
+#include "CoreMinimal.h"
+#endif
diff --git a/Engine/Plugins/Interchange/Editor/Source/Pipelines/Public/InterchangeEditorPipelinesModule.h b/Engine/Plugins/Interchange/Editor/Source/Pipelines/Public/InterchangeEditorPipelinesModule.h
--- a/Engine/Plugins/Interchange/Editor/Source/Pipelines/Public/InterchangeEditorPipelinesModule.h
+++ b/Engine/Plugins/Interchange/Editor/Source/Pipelines/Public/InterchangeEditorPipelinesModule.h
@@ -2,10 +2,7 @@
 
 #pragma once
 
-#include "CoreMinimal.h"
-#include "Modules/ModuleInterface.h"
 #include "Modules/ModuleManager.h"
-#include "Styling/ISlateStyle.h"
 
 class FAssetTypeActions_InterchangeBlueprintPipelineBase;
 class FAssetTypeActions_InterchangePipelineBase;
@@ -47,3 +44,8 @@ public:
 	TSharedPtr<FAssetTypeActions_InterchangePythonPipelineBase> PythonPipelineBase_TypeActions;
 };
 
+
+#if UE_ENABLE_INCLUDE_ORDER_DEPRECATED_IN_5_2
+#include "CoreMinimal.h"
+#include "Styling/ISlateStyle.h"
+#endif
diff --git a/Engine/Plugins/Interchange/Editor/Source/Pipelines/Public/InterchangeGraphInspectorPipeline.h b/Engine/Plugins/Interchange/Editor/Source/Pipelines/Public/InterchangeGraphInspectorPipeline.h
--- a/Engine/Plugins/Interchange/Editor/Source/Pipelines/Public/InterchangeGraphInspectorPipeline.h
+++ b/Engine/Plugins/Interchange/Editor/Source/Pipelines/Public/InterchangeGraphInspectorPipeline.h
@@ -2,12 +2,7 @@
 
 #pragma once
 
-#include "CoreMinimal.h"
 #include "InterchangePipelineBase.h"
-#include "InterchangeSourceData.h"
-#include "Nodes/InterchangeBaseNodeContainer.h"
-#include "UObject/Object.h"
-#include "UObject/ObjectMacros.h"
 
 #include "InterchangeGraphInspectorPipeline.generated.h"
 
@@ -51,3 +46,9 @@ private:
 };
 
 
+
+#if UE_ENABLE_INCLUDE_ORDER_DEPRECATED_IN_5_2
+#include "CoreMinimal.h"
+#include "InterchangeSourceData.h"
+#include "Nodes/InterchangeBaseNodeContainer.h"
+#endif
diff --git a/Engine/Plugins/Interchange/Editor/Source/Pipelines/Public/InterchangePipelineConfigurationGeneric.h b/Engine/Plugins/Interchange/Editor/Source/Pipelines/Public/InterchangePipelineConfigurationGeneric.h
--- a/Engine/Plugins/Interchange/Editor/Source/Pipelines/Public/InterchangePipelineConfigurationGeneric.h
+++ b/Engine/Plugins/Interchange/Editor/Source/Pipelines/Public/InterchangePipelineConfigurationGeneric.h
@@ -2,10 +2,7 @@
 
 #pragma once
 
-#include "CoreMinimal.h"
 #include "InterchangePipelineConfigurationBase.h"
-#include "UObject/Object.h"
-#include "UObject/ObjectMacros.h"
 
 #include "InterchangePipelineConfigurationGeneric.generated.h"
 
@@ -30,3 +27,7 @@ protected:
 		, TArray<UInterchangePipelineBase*>& OutPipelines
 		, TWeakObjectPtr<UInterchangeSourceData> SourceData) override;
 };
+
+#if UE_ENABLE_INCLUDE_ORDER_DEPRECATED_IN_5_2
+#include "CoreMinimal.h"
+#endif
diff --git a/Engine/Plugins/Interchange/Editor/Source/ResultsBrowser/Public/InterchangeResultsBrowserModule.h b/Engine/Plugins/Interchange/Editor/Source/ResultsBrowser/Public/InterchangeResultsBrowserModule.h
--- a/Engine/Plugins/Interchange/Editor/Source/ResultsBrowser/Public/InterchangeResultsBrowserModule.h
+++ b/Engine/Plugins/Interchange/Editor/Source/ResultsBrowser/Public/InterchangeResultsBrowserModule.h
@@ -2,10 +2,7 @@
 
 #pragma once
 
-#include "CoreMinimal.h"
-#include "Modules/ModuleInterface.h"
 #include "Modules/ModuleManager.h"
-#include "Styling/ISlateStyle.h"
 
 
 /**
@@ -37,3 +34,8 @@ public:
 		return FModuleManager::Get().IsModuleLoaded("InterchangeResultsBrowser");
 	}
 };
+
+#if UE_ENABLE_INCLUDE_ORDER_DEPRECATED_IN_5_2
+#include "CoreMinimal.h"
+#include "Styling/ISlateStyle.h"
+#endif
diff --git a/Engine/Plugins/Interchange/Editor/Source/Utilities/Public/InterchangeOpenFileDialog.h b/Engine/Plugins/Interchange/Editor/Source/Utilities/Public/InterchangeOpenFileDialog.h
--- a/Engine/Plugins/Interchange/Editor/Source/Utilities/Public/InterchangeOpenFileDialog.h
+++ b/Engine/Plugins/Interchange/Editor/Source/Utilities/Public/InterchangeOpenFileDialog.h
@@ -2,11 +2,7 @@
 
 #pragma once
 
-#include "Containers/UnrealString.h"
 #include "InterchangeFilePickerBase.h"
-#include "InterchangeTranslatorBase.h"
-#include "UObject/Object.h"
-#include "UObject/ObjectMacros.h"
 
 #include "InterchangeOpenFileDialog.generated.h"
 
@@ -22,3 +18,7 @@ protected:
 	virtual bool FilePickerForTranslatorAssetType(const EInterchangeTranslatorAssetType TranslatorAssetType, FInterchangeFilePickerParameters& Parameters, TArray<FString>& OutFilenames) override;
 	virtual bool FilePickerForTranslatorType(const EInterchangeTranslatorType TranslatorType, FInterchangeFilePickerParameters& Parameters, TArray<FString>& OutFilenames) override;
 };
+
+#if UE_ENABLE_INCLUDE_ORDER_DEPRECATED_IN_5_2
+#include "InterchangeTranslatorBase.h"
+#endif
diff --git a/Engine/Plugins/Interchange/Runtime/Source/Import/Private/Gltf/InterchangeGltfAnimation.cpp b/Engine/Plugins/Interchange/Runtime/Source/Import/Private/Gltf/InterchangeGltfAnimation.cpp
--- a/Engine/Plugins/Interchange/Runtime/Source/Import/Private/Gltf/InterchangeGltfAnimation.cpp
+++ b/Engine/Plugins/Interchange/Runtime/Source/Import/Private/Gltf/InterchangeGltfAnimation.cpp
@@ -3,10 +3,11 @@
 #include "Gltf/InterchangeGltfPrivate.h"
 
 #include "Animation/InterchangeAnimationPayload.h"
-#include "InterchangeImportLog.h"
-
+#include "GLTFAccessor.h"
+#include "GLTFAnimation.h"
 #include "GLTFAsset.h"
-
+#include "GLTFNode.h"
+#include "InterchangeImportLog.h"
 #include "Animation/AnimTypes.h"
 
 namespace UE::Interchange::Gltf::Private
diff --git a/Engine/Plugins/Interchange/Runtime/Source/Import/Private/Gltf/InterchangeGltfMesh.cpp b/Engine/Plugins/Interchange/Runtime/Source/Import/Private/Gltf/InterchangeGltfMesh.cpp
--- a/Engine/Plugins/Interchange/Runtime/Source/Import/Private/Gltf/InterchangeGltfMesh.cpp
+++ b/Engine/Plugins/Interchange/Runtime/Source/Import/Private/Gltf/InterchangeGltfMesh.cpp
@@ -3,8 +3,10 @@
 #include "Gltf/InterchangeGltfPrivate.h"
 
 #include "GLTFAsset.h"
+#include "GLTFMaterial.h"
+#include "GLTFMesh.h"
 #include "GLTFMeshFactory.h"
-
+#include "GLTFNode.h"
 #include "StaticMeshAttributes.h"
 #include "SkeletalMeshAttributes.h"
 
diff --git a/Engine/Plugins/Interchange/Runtime/Source/Import/Private/Gltf/InterchangeGltfTranslator.cpp b/Engine/Plugins/Interchange/Runtime/Source/Import/Private/Gltf/InterchangeGltfTranslator.cpp
--- a/Engine/Plugins/Interchange/Runtime/Source/Import/Private/Gltf/InterchangeGltfTranslator.cpp
+++ b/Engine/Plugins/Interchange/Runtime/Source/Import/Private/Gltf/InterchangeGltfTranslator.cpp
@@ -2,9 +2,14 @@
 
 #include "Gltf/InterchangeGltfTranslator.h"
 
+#include "GLTFAccessor.h"
+#include "GLTFAnimation.h"
 #include "GLTFAsset.h"
+#include "GLTFMesh.h"
 #include "GLTFMeshFactory.h"
+#include "GLTFNode.h"
 #include "GLTFReader.h"
+#include "GLTFTexture.h"
 
 #include "InterchangeAnimationTrackSetNode.h"
 #include "InterchangeCameraNode.h"
diff --git a/Engine/Plugins/Interchange/Runtime/Source/Import/Public/Gltf/InterchangeGltfTranslator.h b/Engine/Plugins/Interchange/Runtime/Source/Import/Public/Gltf/InterchangeGltfTranslator.h
--- a/Engine/Plugins/Interchange/Runtime/Source/Import/Public/Gltf/InterchangeGltfTranslator.h
+++ b/Engine/Plugins/Interchange/Runtime/Source/Import/Public/Gltf/InterchangeGltfTranslator.h
@@ -2,6 +2,7 @@
 #pragma once
 
 #include "GLTFAsset.h"
+#include "GLTFMaterial.h"
 #include "InterchangeTranslatorBase.h"
 #include "Animation/InterchangeAnimationPayloadInterface.h"
 #include "Mesh/InterchangeStaticMeshPayload.h"
diff --git a/Engine/Plugins/MemoryUsageQueries/Source/MemoryUsageQueries/Private/MemoryUsageQueries.cpp b/Engine/Plugins/MemoryUsageQueries/Source/MemoryUsageQueries/Private/MemoryUsageQueries.cpp
--- a/Engine/Plugins/MemoryUsageQueries/Source/MemoryUsageQueries/Private/MemoryUsageQueries.cpp
+++ b/Engine/Plugins/MemoryUsageQueries/Source/MemoryUsageQueries/Private/MemoryUsageQueries.cpp
@@ -9,7 +9,6 @@
 #include "Engine/BlueprintGeneratedClass.h"
 #include "Misc/PackageName.h"
 #include "Modules/ModuleManager.h"
-#include "UObject/GCObject.h"
 #include "Templates/Greater.h"
 #include "UObject/UObjectAllocator.h"
 #include "UObject/UObjectIterator.h"
diff --git a/Engine/Plugins/MeshPainting/Source/MeshPaintEditorMode/Public/MeshPaintEditorMode.h b/Engine/Plugins/MeshPainting/Source/MeshPaintEditorMode/Public/MeshPaintEditorMode.h
--- a/Engine/Plugins/MeshPainting/Source/MeshPaintEditorMode/Public/MeshPaintEditorMode.h
+++ b/Engine/Plugins/MeshPainting/Source/MeshPaintEditorMode/Public/MeshPaintEditorMode.h
@@ -2,4 +2,7 @@
 
 #pragma once
 
+
+#if UE_ENABLE_INCLUDE_ORDER_DEPRECATED_IN_5_2
 #include "CoreMinimal.h"
+#endif
diff --git a/Engine/Plugins/MeshPainting/Source/MeshPaintEditorMode/Public/MeshPaintModeCommands.h b/Engine/Plugins/MeshPainting/Source/MeshPaintEditorMode/Public/MeshPaintModeCommands.h
--- a/Engine/Plugins/MeshPainting/Source/MeshPaintEditorMode/Public/MeshPaintModeCommands.h
+++ b/Engine/Plugins/MeshPainting/Source/MeshPaintEditorMode/Public/MeshPaintModeCommands.h
@@ -2,10 +2,9 @@
 
 #pragma once
 
-#include "Framework/Commands/UICommandInfo.h"
-#include "Framework/Commands/Commands.h"
 #include "Styling/AppStyle.h"
 #include "Tools/InteractiveToolsCommands.h"
+#include "Tools/StandardToolModeCommands.h"
 
 
 class FMeshPaintEditorModeCommands : public TCommands<FMeshPaintEditorModeCommands>
@@ -98,4 +97,4 @@ public:
 	 * @param bUnbind if true, commands are removed, otherwise added
 	 */
 	static void UpdateToolCommandBinding(UInteractiveTool* Tool, TSharedPtr<FUICommandList> UICommandList, bool bUnbind = false);
-};
\ No newline at end of file
+};
diff --git a/Engine/Plugins/MeshPainting/Source/MeshPaintEditorMode/Public/MeshPaintModeHelpers.h b/Engine/Plugins/MeshPainting/Source/MeshPaintEditorMode/Public/MeshPaintModeHelpers.h
--- a/Engine/Plugins/MeshPainting/Source/MeshPaintEditorMode/Public/MeshPaintModeHelpers.h
+++ b/Engine/Plugins/MeshPainting/Source/MeshPaintEditorMode/Public/MeshPaintModeHelpers.h
@@ -2,11 +2,7 @@
 
 #pragma once
 
-#include "CoreMinimal.h"
-#include "Engine/EngineTypes.h"
-#include "MeshPaintingToolsetTypes.h"
 #include "EditorSubsystem.h"
-#include "IMeshPaintComponentAdapter.h"
 
 class FMeshPaintParameters;
 class UImportVertexColorOptions;
@@ -73,4 +69,11 @@ public:
 	void SwapVertexColors();
 	void SaveModifiedTextures();
 	bool CanSaveModifiedTextures();
-};
\ No newline at end of file
+};
+
+#if UE_ENABLE_INCLUDE_ORDER_DEPRECATED_IN_5_2
+#include "CoreMinimal.h"
+#include "Engine/EngineTypes.h"
+#include "IMeshPaintComponentAdapter.h"
+#include "MeshPaintingToolsetTypes.h"
+#endif
diff --git a/Engine/Plugins/MeshPainting/Source/MeshPaintingToolset/Private/BaseMeshPaintComponentAdapter.cpp b/Engine/Plugins/MeshPainting/Source/MeshPaintingToolset/Private/BaseMeshPaintComponentAdapter.cpp
--- a/Engine/Plugins/MeshPainting/Source/MeshPaintingToolset/Private/BaseMeshPaintComponentAdapter.cpp
+++ b/Engine/Plugins/MeshPainting/Source/MeshPaintingToolset/Private/BaseMeshPaintComponentAdapter.cpp
@@ -2,7 +2,6 @@
 
 #include "BaseMeshPaintComponentAdapter.h"
 #include "Distance/DistLine3Segment3.h"
-#include "Spatial/MeshAABBTree3.h"
 #include "Distance/DistLine3Triangle3.h"
 #include "Distance/DistSegment3Triangle3.h"
 
diff --git a/Engine/Plugins/MeshPainting/Source/MeshPaintingToolset/Private/MeshPaintGeometryCollectionAdapter.cpp b/Engine/Plugins/MeshPainting/Source/MeshPaintingToolset/Private/MeshPaintGeometryCollectionAdapter.cpp
--- a/Engine/Plugins/MeshPainting/Source/MeshPaintingToolset/Private/MeshPaintGeometryCollectionAdapter.cpp
+++ b/Engine/Plugins/MeshPainting/Source/MeshPaintingToolset/Private/MeshPaintGeometryCollectionAdapter.cpp
@@ -5,7 +5,6 @@
 #include "GeometryCollection/GeometryCollectionComponent.h"
 #include "GeometryCollection/GeometryCollectionObject.h"
 #include "RenderingThread.h"
-#include "Spatial/MeshAABBTree3.h"
 
 DEFINE_LOG_CATEGORY_STATIC(LogMeshPaintGeometryCollectionComponentAdapter, Log, All);
 //////////////////////////////////////////////////////////////////////////
diff --git a/Engine/Plugins/MeshPainting/Source/MeshPaintingToolset/Private/MeshPaintSplineMeshAdapter.cpp b/Engine/Plugins/MeshPainting/Source/MeshPaintingToolset/Private/MeshPaintSplineMeshAdapter.cpp
--- a/Engine/Plugins/MeshPainting/Source/MeshPaintingToolset/Private/MeshPaintSplineMeshAdapter.cpp
+++ b/Engine/Plugins/MeshPainting/Source/MeshPaintingToolset/Private/MeshPaintSplineMeshAdapter.cpp
@@ -2,10 +2,8 @@
 
 #include "MeshPaintSplineMeshAdapter.h"
 
-#include "Spatial/SpatialInterfaces.h"
 #include "StaticMeshResources.h"
 #include "Components/SplineMeshComponent.h"
-#include "Util/DynamicVector.h"
 
 //////////////////////////////////////////////////////////////////////////
 // FMeshPaintGeometryAdapterForSplineMeshes
diff --git a/Engine/Plugins/MeshPainting/Source/MeshPaintingToolset/Public/BaseMeshPaintComponentAdapter.h b/Engine/Plugins/MeshPainting/Source/MeshPaintingToolset/Public/BaseMeshPaintComponentAdapter.h
--- a/Engine/Plugins/MeshPainting/Source/MeshPaintingToolset/Public/BaseMeshPaintComponentAdapter.h
+++ b/Engine/Plugins/MeshPainting/Source/MeshPaintingToolset/Public/BaseMeshPaintComponentAdapter.h
@@ -5,7 +5,7 @@
 #include "IMeshPaintComponentAdapter.h"
 #include "TMeshPaintOctree.h"
 #include "MeshAdapter.h"
-#include "Spatial/MeshAABBTree3.h"
+#include "Spatial/MeshAABBTree3.h" // IWYU pragma: keep
 #include "UObject/GCObject.h"
 
 typedef UE::Geometry::TIndexMeshArrayAdapter<uint32, double> FIndexMeshArrayAdapterd;
diff --git a/Engine/Plugins/MeshPainting/Source/MeshPaintingToolset/Public/IMeshPaintComponentAdapter.h b/Engine/Plugins/MeshPainting/Source/MeshPaintingToolset/Public/IMeshPaintComponentAdapter.h
--- a/Engine/Plugins/MeshPainting/Source/MeshPaintingToolset/Public/IMeshPaintComponentAdapter.h
+++ b/Engine/Plugins/MeshPainting/Source/MeshPaintingToolset/Public/IMeshPaintComponentAdapter.h
@@ -2,7 +2,7 @@
 
 #pragma once
 
-#include "CoreMinimal.h"
+#include "Math/Matrix.h"
 
 class FReferenceCollector;
 class UMeshComponent;
@@ -102,4 +102,8 @@ public:
 	/** Returns the hit result which intersect with the given ray */
 	virtual bool RayIntersectAdapter(UE::Geometry::FIndex3i& HitTriangle, FVector& HitPosition, const FVector Start, const FVector End) const = 0;
 
-};
\ No newline at end of file
+};
+
+#if UE_ENABLE_INCLUDE_ORDER_DEPRECATED_IN_5_2
+#include "CoreMinimal.h"
+#endif
diff --git a/Engine/Plugins/MeshPainting/Source/MeshPaintingToolset/Public/MeshPaintAdapterFactory.h b/Engine/Plugins/MeshPainting/Source/MeshPaintingToolset/Public/MeshPaintAdapterFactory.h
--- a/Engine/Plugins/MeshPainting/Source/MeshPaintingToolset/Public/MeshPaintAdapterFactory.h
+++ b/Engine/Plugins/MeshPainting/Source/MeshPaintingToolset/Public/MeshPaintAdapterFactory.h
@@ -2,7 +2,7 @@
 
 #pragma once
 
-#include "CoreMinimal.h"
+#include "Templates/SharedPointer.h"
 
 class UMeshComponent;
 class IMeshPaintComponentAdapterFactory;
@@ -18,4 +18,8 @@ public:
 	static void InitializeAdapterGlobals();
 	static void AddReferencedObjectsGlobals(FReferenceCollector& Collector);
 	static void CleanupGlobals();
-};
\ No newline at end of file
+};
+
+#if UE_ENABLE_INCLUDE_ORDER_DEPRECATED_IN_5_2
+#include "CoreMinimal.h"
+#endif
diff --git a/Engine/Plugins/MeshPainting/Source/MeshPaintingToolset/Public/MeshPaintComponentAdapterFactory.h b/Engine/Plugins/MeshPainting/Source/MeshPaintingToolset/Public/MeshPaintComponentAdapterFactory.h
--- a/Engine/Plugins/MeshPainting/Source/MeshPaintingToolset/Public/MeshPaintComponentAdapterFactory.h
+++ b/Engine/Plugins/MeshPainting/Source/MeshPaintingToolset/Public/MeshPaintComponentAdapterFactory.h
@@ -2,7 +2,6 @@
 
 #pragma once
 
-#include "CoreTypes.h"
 #include "Templates/SharedPointer.h"
 
 class IMeshPaintComponentAdapter;
@@ -20,4 +19,4 @@ public:
 	virtual void AddReferencedObjectsGlobals(FReferenceCollector& Collector) = 0;
 	virtual void CleanupGlobals() = 0;
 	virtual ~IMeshPaintComponentAdapterFactory() {}
-};
\ No newline at end of file
+};
diff --git a/Engine/Plugins/MeshPainting/Source/MeshPaintingToolset/Public/MeshPaintGeometryCollectionAdapter.h b/Engine/Plugins/MeshPainting/Source/MeshPaintingToolset/Public/MeshPaintGeometryCollectionAdapter.h
--- a/Engine/Plugins/MeshPainting/Source/MeshPaintingToolset/Public/MeshPaintGeometryCollectionAdapter.h
+++ b/Engine/Plugins/MeshPainting/Source/MeshPaintingToolset/Public/MeshPaintGeometryCollectionAdapter.h
@@ -2,10 +2,8 @@
 
 #pragma once
 
-#include "CoreMinimal.h"
 #include "BaseMeshPaintComponentAdapter.h"
 #include "MeshPaintComponentAdapterFactory.h"
-#include "UObject/WeakObjectPtrTemplates.h"
 #include "UObject/WeakObjectPtr.h"
 
 class UBodySetup;
@@ -81,3 +79,7 @@ public:
 	virtual void AddReferencedObjectsGlobals(FReferenceCollector& Collector) override { FMeshPaintGeometryCollectionComponentAdapter::AddReferencedObjectsGlobals(Collector); }
 	virtual void CleanupGlobals() override { FMeshPaintGeometryCollectionComponentAdapter::CleanupGlobals(); }
 };
+
+#if UE_ENABLE_INCLUDE_ORDER_DEPRECATED_IN_5_2
+#include "CoreMinimal.h"
+#endif
diff --git a/Engine/Plugins/MeshPainting/Source/MeshPaintingToolset/Public/MeshPaintHelpers.h b/Engine/Plugins/MeshPainting/Source/MeshPaintingToolset/Public/MeshPaintHelpers.h
--- a/Engine/Plugins/MeshPainting/Source/MeshPaintingToolset/Public/MeshPaintHelpers.h
+++ b/Engine/Plugins/MeshPainting/Source/MeshPaintingToolset/Public/MeshPaintHelpers.h
@@ -2,18 +2,14 @@
 
 #pragma once
 
-#include "CoreMinimal.h"
-#include "Engine/EngineTypes.h"
-#include "MeshPaintingToolsetTypes.h"
-#include "Kismet/BlueprintFunctionLibrary.h"
-#include "IMeshPaintComponentAdapter.h"
-#include "InteractiveToolManager.h"
-#include "EdModeInteractiveToolsContext.h"
 #include "Engine/Engine.h"
 #include "Engine/StaticMesh.h"
-#include "Subsystems/EngineSubsystem.h"
+#include "Math/Ray.h"
+#include "UObject/Package.h"
 #include "MeshPaintHelpers.generated.h"
 
+enum class EMeshPaintModeAction : uint8;
+
 class FMeshPaintParameters;
 class UImportVertexColorOptions;
 class UTexture2D;
@@ -303,3 +299,12 @@ void UMeshPaintingSubsystem::ApplyBrushToVertex(const FVector& VertexPosition, c
 		InOutValue = FMath::LerpStable(OldValue, PaintValue, PaintStrength);
 	}	
 };
+
+#if UE_ENABLE_INCLUDE_ORDER_DEPRECATED_IN_5_2
+#include "CoreMinimal.h"
+#include "EdModeInteractiveToolsContext.h"
+#include "IMeshPaintComponentAdapter.h"
+#include "InteractiveToolManager.h"
+#include "Kismet/BlueprintFunctionLibrary.h"
+#include "MeshPaintingToolsetTypes.h"
+#endif
diff --git a/Engine/Plugins/MeshPainting/Source/MeshPaintingToolset/Public/MeshPaintInteractions.h b/Engine/Plugins/MeshPainting/Source/MeshPaintingToolset/Public/MeshPaintInteractions.h
--- a/Engine/Plugins/MeshPainting/Source/MeshPaintingToolset/Public/MeshPaintInteractions.h
+++ b/Engine/Plugins/MeshPainting/Source/MeshPaintingToolset/Public/MeshPaintInteractions.h
@@ -2,11 +2,13 @@
 
 #pragma once
 
-#include "CoreMinimal.h"
 #include "InteractionMechanic.h"
 #include "UObject/Interface.h"
 #include "MeshPaintInteractions.generated.h"
 
+struct FInputDeviceRay;
+struct FInputRayHit;
+
 class IMeshPaintComponentAdapter;
 class AActor;
 class UMeshComponent;
@@ -49,3 +51,7 @@ protected:
 	TArray<TObjectPtr<AActor>> CachedClickedActors;
 	bool bAddToSelectionSet;
 };
+
+#if UE_ENABLE_INCLUDE_ORDER_DEPRECATED_IN_5_2
+#include "CoreMinimal.h"
+#endif
diff --git a/Engine/Plugins/MeshPainting/Source/MeshPaintingToolset/Public/MeshPaintSkeletalMeshAdapter.h b/Engine/Plugins/MeshPainting/Source/MeshPaintingToolset/Public/MeshPaintSkeletalMeshAdapter.h
--- a/Engine/Plugins/MeshPainting/Source/MeshPaintingToolset/Public/MeshPaintSkeletalMeshAdapter.h
+++ b/Engine/Plugins/MeshPainting/Source/MeshPaintingToolset/Public/MeshPaintSkeletalMeshAdapter.h
@@ -2,8 +2,6 @@
 
 #pragma once
 
-#include "CoreMinimal.h"
-#include "RawIndexBuffer.h"
 #include "Components/SkeletalMeshComponent.h"
 #include "BaseMeshPaintComponentAdapter.h"
 #include "MeshPaintComponentAdapterFactory.h"
@@ -84,3 +82,8 @@ public:
 	virtual void AddReferencedObjectsGlobals(FReferenceCollector& Collector) override { FMeshPaintSkeletalMeshComponentAdapter::AddReferencedObjectsGlobals(Collector); }
 	virtual void CleanupGlobals() override { FMeshPaintSkeletalMeshComponentAdapter::CleanupGlobals(); }
 };
+
+#if UE_ENABLE_INCLUDE_ORDER_DEPRECATED_IN_5_2
+#include "CoreMinimal.h"
+#include "RawIndexBuffer.h"
+#endif
diff --git a/Engine/Plugins/MeshPainting/Source/MeshPaintingToolset/Public/MeshPaintStaticMeshAdapter.h b/Engine/Plugins/MeshPainting/Source/MeshPaintingToolset/Public/MeshPaintStaticMeshAdapter.h
--- a/Engine/Plugins/MeshPainting/Source/MeshPaintingToolset/Public/MeshPaintStaticMeshAdapter.h
+++ b/Engine/Plugins/MeshPainting/Source/MeshPaintingToolset/Public/MeshPaintStaticMeshAdapter.h
@@ -2,9 +2,7 @@
 
 #pragma once
 
-#include "CoreMinimal.h"
 #include "BaseMeshPaintComponentAdapter.h"
-#include "RawIndexBuffer.h"
 #include "Components/StaticMeshComponent.h"
 #include "MeshPaintComponentAdapterFactory.h"
 
@@ -74,3 +72,8 @@ public:
 	virtual void AddReferencedObjectsGlobals(FReferenceCollector& Collector) override { FMeshPaintStaticMeshComponentAdapter::AddReferencedObjectsGlobals(Collector); }
 	virtual void CleanupGlobals() override { FMeshPaintStaticMeshComponentAdapter::CleanupGlobals(); }
 };
+
+#if UE_ENABLE_INCLUDE_ORDER_DEPRECATED_IN_5_2
+#include "CoreMinimal.h"
+#include "RawIndexBuffer.h"
+#endif
diff --git a/Engine/Plugins/MeshPainting/Source/MeshPaintingToolset/Public/MeshPaintingToolset.h b/Engine/Plugins/MeshPainting/Source/MeshPaintingToolset/Public/MeshPaintingToolset.h
--- a/Engine/Plugins/MeshPainting/Source/MeshPaintingToolset/Public/MeshPaintingToolset.h
+++ b/Engine/Plugins/MeshPainting/Source/MeshPaintingToolset/Public/MeshPaintingToolset.h
@@ -2,4 +2,7 @@
 
 #pragma once
 
+
+#if UE_ENABLE_INCLUDE_ORDER_DEPRECATED_IN_5_2
 #include "CoreMinimal.h"
+#endif
diff --git a/Engine/Plugins/MeshPainting/Source/MeshPaintingToolset/Public/MeshPaintingToolsetTypes.h b/Engine/Plugins/MeshPainting/Source/MeshPaintingToolset/Public/MeshPaintingToolsetTypes.h
--- a/Engine/Plugins/MeshPainting/Source/MeshPaintingToolset/Public/MeshPaintingToolsetTypes.h
+++ b/Engine/Plugins/MeshPainting/Source/MeshPaintingToolset/Public/MeshPaintingToolsetTypes.h
@@ -2,12 +2,14 @@
 
 #pragma once
 
-#include "Engine/Texture.h"
-#include "Engine/TextureRenderTarget2D.h"
 
 #include "MeshPaintingToolsetTypes.generated.h"
 
+class UMaterialInterface;
 class UMeshComponent;
+class UTexture2D;
+class UTexture;
+class UTextureRenderTarget2D;
 
 class FMeshPaintParameters;
 DECLARE_MULTICAST_DELEGATE_FourParams(FApplyVertexPaintData, const FMeshPaintParameters& /* &InParams*/, const FLinearColor& /*OldColor*/, FLinearColor& /*NewColor*/, const float /*PaintAmount*/);
@@ -197,4 +199,10 @@ struct FMeshSelectedMaterialInfo
 		: NumMaterials(InNumMaterials)
 		, SelectedMaterialIndex(0)
 	{}
-};
\ No newline at end of file
+};
+
+#if UE_ENABLE_INCLUDE_ORDER_DEPRECATED_IN_5_2
+#include "Engine/Texture.h"
+#include "Engine/TextureRenderTarget2D.h"
+#include "Materials/MaterialInterface.h"
+#endif
diff --git a/Engine/Plugins/MeshPainting/Source/MeshPaintingToolset/Public/MeshTexturePaintingTool.h b/Engine/Plugins/MeshPainting/Source/MeshPaintingToolset/Public/MeshTexturePaintingTool.h
--- a/Engine/Plugins/MeshPainting/Source/MeshPaintingToolset/Public/MeshTexturePaintingTool.h
+++ b/Engine/Plugins/MeshPainting/Source/MeshPaintingToolset/Public/MeshTexturePaintingTool.h
@@ -2,21 +2,18 @@
 
 #pragma once
 
-#include "CoreMinimal.h"
-#include "UObject/NoExportTypes.h"
-#include "InteractiveToolBuilder.h"
 #include "BaseTools/BaseBrushTool.h"
-#include "InteractiveTool.h"
-#include "Delegates/DelegateCombinations.h"
-#include "Components/MeshComponent.h"
 #include "MeshPaintHelpers.h"
+#include "MeshPaintInteractions.h"
 #include "MeshPaintingToolsetTypes.h"
-#include "MeshVertexPaintingTool.h"
 #include "Misc/ITransaction.h"
-#include "TexturePaintToolset.h"
 
 #include "MeshTexturePaintingTool.generated.h"
 
+class UMeshToolManager;
+enum class EToolShutdownType : uint8;
+struct FTexturePaintMeshSectionInfo;
+
 
 struct FToolBuilderState;
 struct FPaintTexture2DData;
@@ -242,3 +239,12 @@ private:
 	/** Hold the transaction while we are painting */
 	TUniquePtr<FScopedTransaction> PaintingTransaction;
 };
+
+#if UE_ENABLE_INCLUDE_ORDER_DEPRECATED_IN_5_2
+#include "Components/MeshComponent.h"
+#include "CoreMinimal.h"
+#include "Delegates/DelegateCombinations.h"
+#include "MeshVertexPaintingTool.h"
+#include "TexturePaintToolset.h"
+#include "UObject/NoExportTypes.h"
+#endif
diff --git a/Engine/Plugins/MeshPainting/Source/MeshPaintingToolset/Public/MeshVertexPaintingTool.h b/Engine/Plugins/MeshPainting/Source/MeshPaintingToolset/Public/MeshVertexPaintingTool.h
--- a/Engine/Plugins/MeshPainting/Source/MeshPaintingToolset/Public/MeshVertexPaintingTool.h
+++ b/Engine/Plugins/MeshPainting/Source/MeshPaintingToolset/Public/MeshVertexPaintingTool.h
@@ -2,18 +2,14 @@
 
 #pragma once
 
-#include "CoreMinimal.h"
-#include "UObject/NoExportTypes.h"
-#include "InteractiveToolBuilder.h"
 #include "BaseTools/BaseBrushTool.h"
-#include "InteractiveTool.h"
-#include "Delegates/DelegateCombinations.h"
-#include "Components/MeshComponent.h"
-#include "MeshPaintHelpers.h"
 #include "MeshPaintingToolsetTypes.h"
 #include "MeshPaintInteractions.h"
 #include "MeshVertexPaintingTool.generated.h"
 
+enum class EToolShutdownType : uint8;
+struct FPerVertexPaintActionArgs;
+
 
 struct FToolBuilderState;
 enum class EMeshPaintModeAction : uint8;
@@ -294,3 +290,10 @@ private:
 	UPROPERTY(Transient)
 	TObjectPtr<UMeshWeightPaintingToolProperties> WeightProperties;
 };
+
+#if UE_ENABLE_INCLUDE_ORDER_DEPRECATED_IN_5_2
+#include "Components/MeshComponent.h"
+#include "CoreMinimal.h"
+#include "MeshPaintHelpers.h"
+#include "UObject/NoExportTypes.h"
+#endif
diff --git a/Engine/Plugins/MeshPainting/Source/MeshPaintingToolset/Public/TexturePaintToolset.h b/Engine/Plugins/MeshPainting/Source/MeshPaintingToolset/Public/TexturePaintToolset.h
--- a/Engine/Plugins/MeshPainting/Source/MeshPaintingToolset/Public/TexturePaintToolset.h
+++ b/Engine/Plugins/MeshPainting/Source/MeshPaintingToolset/Public/TexturePaintToolset.h
@@ -7,6 +7,8 @@
 #include "Kismet/BlueprintFunctionLibrary.h"
 #include "TexturePaintToolset.generated.h"
 
+class FTexture;
+
 class UMeshComponent;
 class UTexture;
 class UTexture2D;
@@ -100,4 +102,4 @@ public:
 
 	/** Checks whether or not the given MeshComponent's materials use the specified texture */
 	static bool DoesMeshComponentUseTexture(UMeshComponent* MeshComponent, UTexture* Texture);
-};
\ No newline at end of file
+};
diff --git a/Engine/Plugins/Messaging/UdpMessaging/Source/UdpMessaging/Public/IUdpMessageTunnel.h b/Engine/Plugins/Messaging/UdpMessaging/Source/UdpMessaging/Public/IUdpMessageTunnel.h
--- a/Engine/Plugins/Messaging/UdpMessaging/Source/UdpMessaging/Public/IUdpMessageTunnel.h
+++ b/Engine/Plugins/Messaging/UdpMessaging/Source/UdpMessaging/Public/IUdpMessageTunnel.h
@@ -2,10 +2,7 @@
 
 #pragma once
 
-#include "CoreTypes.h"
-#include "Containers/Array.h"
 #include "Delegates/Delegate.h"
-#include "Templates/SharedPointer.h"
 
 class IUdpMessageTunnelConnection;
 
diff --git a/Engine/Plugins/Messaging/UdpMessaging/Source/UdpMessaging/Public/IUdpMessageTunnelConnection.h b/Engine/Plugins/Messaging/UdpMessaging/Source/UdpMessaging/Public/IUdpMessageTunnelConnection.h
--- a/Engine/Plugins/Messaging/UdpMessaging/Source/UdpMessaging/Public/IUdpMessageTunnelConnection.h
+++ b/Engine/Plugins/Messaging/UdpMessaging/Source/UdpMessaging/Public/IUdpMessageTunnelConnection.h
@@ -3,8 +3,9 @@
 #pragma once
 
 #include "CoreTypes.h"
-#include "Internationalization/Text.h"
-#include "Misc/Timespan.h"
+
+class FText;
+struct FTimespan;
 
 
 /**
@@ -57,3 +58,8 @@ public:
 	/** Virtual destructor. */
 	virtual ~IUdpMessageTunnelConnection() { }
 };
+
+#if UE_ENABLE_INCLUDE_ORDER_DEPRECATED_IN_5_2
+#include "Internationalization/Text.h"
+#include "Misc/Timespan.h"
+#endif
diff --git a/Engine/Plugins/Messaging/UdpMessaging/Source/UdpMessaging/Public/Shared/UdpMessagingSettings.h b/Engine/Plugins/Messaging/UdpMessaging/Source/UdpMessaging/Public/Shared/UdpMessagingSettings.h
--- a/Engine/Plugins/Messaging/UdpMessaging/Source/UdpMessaging/Public/Shared/UdpMessagingSettings.h
+++ b/Engine/Plugins/Messaging/UdpMessaging/Source/UdpMessaging/Public/Shared/UdpMessagingSettings.h
@@ -2,11 +2,6 @@
 
 #pragma once
 
-#include "CoreTypes.h"
-#include "Containers/Array.h"
-#include "Containers/UnrealString.h"
-#include "UObject/ObjectMacros.h"
-#include "UObject/Object.h"
 
 #include "UdpMessagingSettings.generated.h"
 
diff --git a/Engine/Plugins/Messaging/UdpMessaging/Source/UdpMessaging/Public/UdpMessagingTracing.h b/Engine/Plugins/Messaging/UdpMessaging/Source/UdpMessaging/Public/UdpMessagingTracing.h
--- a/Engine/Plugins/Messaging/UdpMessaging/Source/UdpMessaging/Public/UdpMessagingTracing.h
+++ b/Engine/Plugins/Messaging/UdpMessaging/Source/UdpMessaging/Public/UdpMessagingTracing.h
@@ -2,12 +2,17 @@
 
 #pragma once
 
-#include "Logging/LogMacros.h"
 #include "ProfilingDebugging/CpuProfilerTrace.h"
 
+namespace UE::Trace { class FChannel; }
+
 #if UE_TRACE_ENABLED && !UE_BUILD_SHIPPING
 UE_TRACE_CHANNEL_EXTERN(MessagingChannel, UDPMESSAGING_API);
 #define SCOPED_MESSAGING_TRACE(TraceName) TRACE_CPUPROFILER_EVENT_SCOPE_ON_CHANNEL(TraceName, MessagingChannel)
 #else
 #define SCOPED_MESSAGING_TRACE(TraceName)
 #endif
+
+#if UE_ENABLE_INCLUDE_ORDER_DEPRECATED_IN_5_2
+#include "Logging/LogMacros.h"
+#endif
diff --git a/Engine/Plugins/MovieScene/LevelSequenceEditor/Source/LevelSequenceEditor/Public/LevelSequenceEditorBlueprintLibrary.h b/Engine/Plugins/MovieScene/LevelSequenceEditor/Source/LevelSequenceEditor/Public/LevelSequenceEditorBlueprintLibrary.h
--- a/Engine/Plugins/MovieScene/LevelSequenceEditor/Source/LevelSequenceEditor/Public/LevelSequenceEditorBlueprintLibrary.h
+++ b/Engine/Plugins/MovieScene/LevelSequenceEditor/Source/LevelSequenceEditor/Public/LevelSequenceEditorBlueprintLibrary.h
@@ -2,13 +2,13 @@
 
 #pragma once
 
-#include "Misc/QualifiedFrameTime.h"
 #include "Kismet/BlueprintFunctionLibrary.h"
-#include "MovieSceneSequencePlayer.h"
-#include "MovieSceneObjectBindingID.h"
-#include "MovieSceneBindingProxy.h"
 #include "LevelSequenceEditorBlueprintLibrary.generated.h"
 
+struct FMovieSceneBindingProxy;
+struct FMovieSceneObjectBindingID;
+struct FMovieSceneSequencePlaybackParams;
+
 class ISequencer;
 class ULevelSequence;
 class UMovieSceneFolder;
@@ -290,4 +290,11 @@ public:
 	 * NOTE: Only to be called by LevelSequenceEditor::Construct.
 	 */
 	static void SetSequencer(TSharedRef<ISequencer> InSequencer);
-};
\ No newline at end of file
+};
+
+#if UE_ENABLE_INCLUDE_ORDER_DEPRECATED_IN_5_2
+#include "Misc/QualifiedFrameTime.h"
+#include "MovieSceneBindingProxy.h"
+#include "MovieSceneObjectBindingID.h"
+#include "MovieSceneSequencePlayer.h"
+#endif
diff --git a/Engine/Plugins/MovieScene/LevelSequenceEditor/Source/LevelSequenceEditor/Public/LevelSequenceEditorModule.h b/Engine/Plugins/MovieScene/LevelSequenceEditor/Source/LevelSequenceEditor/Public/LevelSequenceEditorModule.h
--- a/Engine/Plugins/MovieScene/LevelSequenceEditor/Source/LevelSequenceEditor/Public/LevelSequenceEditorModule.h
+++ b/Engine/Plugins/MovieScene/LevelSequenceEditor/Source/LevelSequenceEditor/Public/LevelSequenceEditorModule.h
@@ -2,7 +2,7 @@
 
 #pragma once
 
-#include "CoreMinimal.h"
+#include "Delegates/Delegate.h"
 #include "Modules/ModuleInterface.h"
 
 
@@ -22,3 +22,7 @@ public:
 	DECLARE_MULTICAST_DELEGATE_OneParam(FAllowPlaybackContext, bool&);
 	virtual FAllowPlaybackContext& OnComputePlaybackContext() = 0;
 };
+
+#if UE_ENABLE_INCLUDE_ORDER_DEPRECATED_IN_5_2
+#include "CoreMinimal.h"
+#endif
diff --git a/Engine/Plugins/MovieScene/LevelSequenceEditor/Source/LevelSequenceEditor/Public/LevelSequenceEditorSubsystem.h b/Engine/Plugins/MovieScene/LevelSequenceEditor/Source/LevelSequenceEditor/Public/LevelSequenceEditorSubsystem.h
--- a/Engine/Plugins/MovieScene/LevelSequenceEditor/Source/LevelSequenceEditor/Public/LevelSequenceEditorSubsystem.h
+++ b/Engine/Plugins/MovieScene/LevelSequenceEditor/Source/LevelSequenceEditor/Public/LevelSequenceEditorSubsystem.h
@@ -3,16 +3,21 @@
 
 #pragma once
 
-#include "CoreMinimal.h"
 #include "EditorSubsystem.h"
-#include "ISequencer.h"
 
 #include "SequenceTimeUnit.h"
-#include "MovieSceneBindingProxy.h"
-#include "SequencerUtilities.h"
 
 #include "LevelSequenceEditorSubsystem.generated.h"
 
+class FUICommandList;
+class ISequencer;
+class UMovieSceneTrack;
+struct FMovieSceneBindingProxy;
+struct FMovieScenePasteBindingsParams;
+struct FMovieScenePasteFoldersParams;
+struct FMovieScenePasteSectionsParams;
+struct FMovieScenePasteTracksParams;
+
 DECLARE_LOG_CATEGORY_EXTERN(LogLevelSequenceEditor, Log, All);
 
 class ACineCameraActor;
@@ -205,3 +210,10 @@ private:
 	TSharedPtr<FExtender> AssignActorMenuExtender;
 	TSharedPtr<FExtender> RebindComponentMenuExtender;
 };
+
+#if UE_ENABLE_INCLUDE_ORDER_DEPRECATED_IN_5_2
+#include "CoreMinimal.h"
+#include "ISequencer.h"
+#include "MovieSceneBindingProxy.h"
+#include "SequencerUtilities.h"
+#endif
diff --git a/Engine/Plugins/MovieScene/LevelSequenceEditor/Source/LevelSequenceEditor/Public/Misc/LevelSequenceEditorSpawnRegister.h b/Engine/Plugins/MovieScene/LevelSequenceEditor/Source/LevelSequenceEditor/Public/Misc/LevelSequenceEditorSpawnRegister.h
--- a/Engine/Plugins/MovieScene/LevelSequenceEditor/Source/LevelSequenceEditor/Public/Misc/LevelSequenceEditorSpawnRegister.h
+++ b/Engine/Plugins/MovieScene/LevelSequenceEditor/Source/LevelSequenceEditor/Public/Misc/LevelSequenceEditorSpawnRegister.h
@@ -2,13 +2,11 @@
 
 #pragma once
 
-#include "CoreMinimal.h"
-#include "Templates/ValueOrError.h"
-#include "MovieSceneSequenceID.h"
-#include "MovieSceneSpawnRegister.h"
 #include "LevelSequenceSpawnRegister.h"
 #include "UObject/ObjectKey.h"
 
+struct FNewSpawnable;
+
 class IMovieScenePlayer;
 class ISequencer;
 class UMovieScene;
@@ -107,3 +105,7 @@ private:
 	/** Handle to a delegate that is bound to FCoreUObjectDelegates::OnObjectPreSave to harvest changes to spawned objects. */
 	FDelegateHandle OnObjectSavedHandle;
 };
+
+#if UE_ENABLE_INCLUDE_ORDER_DEPRECATED_IN_5_2
+#include "CoreMinimal.h"
+#endif
diff --git a/Engine/Plugins/MovieScene/SequencerScripting/Source/SequencerScripting/Public/ExtensionLibraries/MovieSceneBindingExtensions.h b/Engine/Plugins/MovieScene/SequencerScripting/Source/SequencerScripting/Public/ExtensionLibraries/MovieSceneBindingExtensions.h
--- a/Engine/Plugins/MovieScene/SequencerScripting/Source/SequencerScripting/Public/ExtensionLibraries/MovieSceneBindingExtensions.h
+++ b/Engine/Plugins/MovieScene/SequencerScripting/Source/SequencerScripting/Public/ExtensionLibraries/MovieSceneBindingExtensions.h
@@ -2,15 +2,14 @@
 
 #pragma once
 
-#include "CoreTypes.h"
-#include "UObject/ObjectMacros.h"
 #include "Kismet/BlueprintFunctionLibrary.h"
-#include "MovieSceneTrack.h"
-#include "Templates/SubclassOf.h"
-#include "MovieSceneBindingProxy.h"
 
 #include "MovieSceneBindingExtensions.generated.h"
 
+class UMovieSceneTrack;
+struct FMovieSceneBindingProxy;
+template <typename T> class TSubclassOf;
+
 
 /**
  * Function library containing methods that should be hoisted onto FMovieSceneBindingProxies for scripting
@@ -198,4 +197,10 @@ public:
 	 */
 	UFUNCTION(BlueprintCallable, Category = "Sequencer|Sequence", meta = (ScriptMethod))
 	static void MoveBindingContents(const FMovieSceneBindingProxy& SourceBindingId, const FMovieSceneBindingProxy& DestinationBindingId);
-};
\ No newline at end of file
+};
+
+#if UE_ENABLE_INCLUDE_ORDER_DEPRECATED_IN_5_2
+#include "MovieSceneBindingProxy.h"
+#include "MovieSceneTrack.h"
+#include "Templates/SubclassOf.h"
+#endif
diff --git a/Engine/Plugins/MovieScene/SequencerScripting/Source/SequencerScripting/Public/ExtensionLibraries/MovieSceneEventTrackExtensions.h b/Engine/Plugins/MovieScene/SequencerScripting/Source/SequencerScripting/Public/ExtensionLibraries/MovieSceneEventTrackExtensions.h
--- a/Engine/Plugins/MovieScene/SequencerScripting/Source/SequencerScripting/Public/ExtensionLibraries/MovieSceneEventTrackExtensions.h
+++ b/Engine/Plugins/MovieScene/SequencerScripting/Source/SequencerScripting/Public/ExtensionLibraries/MovieSceneEventTrackExtensions.h
@@ -2,8 +2,6 @@
 
 #pragma once
 
-#include "CoreTypes.h"
-#include "UObject/ObjectMacros.h"
 #include "Kismet/BlueprintFunctionLibrary.h"
 #include "Channels/MovieSceneEvent.h"
 
diff --git a/Engine/Plugins/MovieScene/SequencerScripting/Source/SequencerScripting/Public/ExtensionLibraries/MovieSceneFolderExtensions.h b/Engine/Plugins/MovieScene/SequencerScripting/Source/SequencerScripting/Public/ExtensionLibraries/MovieSceneFolderExtensions.h
--- a/Engine/Plugins/MovieScene/SequencerScripting/Source/SequencerScripting/Public/ExtensionLibraries/MovieSceneFolderExtensions.h
+++ b/Engine/Plugins/MovieScene/SequencerScripting/Source/SequencerScripting/Public/ExtensionLibraries/MovieSceneFolderExtensions.h
@@ -2,13 +2,12 @@
 
 #pragma once
 
-#include "CoreTypes.h"
-#include "UObject/ObjectMacros.h"
 #include "Kismet/BlueprintFunctionLibrary.h"
-#include "MovieSceneBindingProxy.h"
 
 #include "MovieSceneFolderExtensions.generated.h"
 
+struct FMovieSceneBindingProxy;
+
 class UMovieSceneFolder;
 class UMovieSceneTrack;
 
@@ -157,4 +156,8 @@ public:
 	 */
 	UFUNCTION(BlueprintCallable, Category = "Sequencer|Sequence", meta=(ScriptMethod))
 	static bool RemoveChildObjectBinding(UMovieSceneFolder* Folder, const FMovieSceneBindingProxy InObjectBinding);
-};
\ No newline at end of file
+};
+
+#if UE_ENABLE_INCLUDE_ORDER_DEPRECATED_IN_5_2
+#include "MovieSceneBindingProxy.h"
+#endif
diff --git a/Engine/Plugins/MovieScene/SequencerScripting/Source/SequencerScripting/Public/ExtensionLibraries/MovieSceneMaterialTrackExtensions.h b/Engine/Plugins/MovieScene/SequencerScripting/Source/SequencerScripting/Public/ExtensionLibraries/MovieSceneMaterialTrackExtensions.h
--- a/Engine/Plugins/MovieScene/SequencerScripting/Source/SequencerScripting/Public/ExtensionLibraries/MovieSceneMaterialTrackExtensions.h
+++ b/Engine/Plugins/MovieScene/SequencerScripting/Source/SequencerScripting/Public/ExtensionLibraries/MovieSceneMaterialTrackExtensions.h
@@ -2,12 +2,12 @@
 
 #pragma once
 
-#include "CoreTypes.h"
-#include "UObject/ObjectMacros.h"
 #include "Kismet/BlueprintFunctionLibrary.h"
 
 #include "MovieSceneMaterialTrackExtensions.generated.h"
 
+class UMovieSceneComponentMaterialTrack;
+
 class UMovieSceneMaterialTrack;
 
 /**
diff --git a/Engine/Plugins/MovieScene/SequencerScripting/Source/SequencerScripting/Public/ExtensionLibraries/MovieScenePrimitiveMaterialTrackExtensions.h b/Engine/Plugins/MovieScene/SequencerScripting/Source/SequencerScripting/Public/ExtensionLibraries/MovieScenePrimitiveMaterialTrackExtensions.h
--- a/Engine/Plugins/MovieScene/SequencerScripting/Source/SequencerScripting/Public/ExtensionLibraries/MovieScenePrimitiveMaterialTrackExtensions.h
+++ b/Engine/Plugins/MovieScene/SequencerScripting/Source/SequencerScripting/Public/ExtensionLibraries/MovieScenePrimitiveMaterialTrackExtensions.h
@@ -2,8 +2,6 @@
 
 #pragma once
 
-#include "CoreTypes.h"
-#include "UObject/ObjectMacros.h"
 #include "Kismet/BlueprintFunctionLibrary.h"
 
 #include "MovieScenePrimitiveMaterialTrackExtensions.generated.h"
diff --git a/Engine/Plugins/MovieScene/SequencerScripting/Source/SequencerScripting/Public/ExtensionLibraries/MovieSceneSectionExtensions.h b/Engine/Plugins/MovieScene/SequencerScripting/Source/SequencerScripting/Public/ExtensionLibraries/MovieSceneSectionExtensions.h
--- a/Engine/Plugins/MovieScene/SequencerScripting/Source/SequencerScripting/Public/ExtensionLibraries/MovieSceneSectionExtensions.h
+++ b/Engine/Plugins/MovieScene/SequencerScripting/Source/SequencerScripting/Public/ExtensionLibraries/MovieSceneSectionExtensions.h
@@ -2,14 +2,14 @@
 
 #pragma once
 
-#include "CoreTypes.h"
-#include "Templates/SubclassOf.h"
-#include "UObject/ObjectMacros.h"
 #include "Kismet/BlueprintFunctionLibrary.h"
-#include "SequencerScriptingRange.h"
-#include "MovieSceneScriptingChannel.h"
 #include "MovieSceneSectionExtensions.generated.h"
 
+class UMovieSceneScriptingChannel;
+class UMovieSceneSequence;
+class UMovieSceneSubSection;
+struct FSequencerScriptingRange;
+
 class UMovieSceneSection;
 
 /**
@@ -269,3 +269,9 @@ public:
 	UFUNCTION(BlueprintCallable, Category = "Sequencer|Section", meta = (ScriptMethod))
 	static int32 GetParentSequenceFrame(UMovieSceneSubSection* Section, int32 InFrame, UMovieSceneSequence* ParentSequence);
 };
+
+#if UE_ENABLE_INCLUDE_ORDER_DEPRECATED_IN_5_2
+#include "MovieSceneScriptingChannel.h"
+#include "SequencerScriptingRange.h"
+#include "Templates/SubclassOf.h"
+#endif
diff --git a/Engine/Plugins/MovieScene/SequencerScripting/Source/SequencerScripting/Public/ExtensionLibraries/MovieSceneSequenceExtensions.h b/Engine/Plugins/MovieScene/SequencerScripting/Source/SequencerScripting/Public/ExtensionLibraries/MovieSceneSequenceExtensions.h
--- a/Engine/Plugins/MovieScene/SequencerScripting/Source/SequencerScripting/Public/ExtensionLibraries/MovieSceneSequenceExtensions.h
+++ b/Engine/Plugins/MovieScene/SequencerScripting/Source/SequencerScripting/Public/ExtensionLibraries/MovieSceneSequenceExtensions.h
@@ -2,19 +2,23 @@
 
 #pragma once
 
-#include "CoreTypes.h"
-#include "UObject/ObjectMacros.h"
-#include "Containers/ArrayView.h"
 #include "Kismet/BlueprintFunctionLibrary.h"
-#include "MovieSceneBindingProxy.h"
-#include "SequencerScriptingRange.h"
-#include "Templates/SubclassOf.h"
-#include "MovieSceneTrack.h"
 #include "MovieSceneObjectBindingID.h" // for EMovieSceneObjectBindingSpace
-#include "MovieScene.h" // only for FMovieSceneMarkedFrame in the .gen.cpp
 
 #include "MovieSceneSequenceExtensions.generated.h"
 
+class UMovieScene;
+class UMovieSceneFolder;
+class UMovieSceneTrack;
+enum class EMovieSceneEvaluationType : uint8;
+enum class EUpdateClockSource : uint8;
+struct FFrameRate;
+struct FMovieSceneBindingProxy;
+struct FMovieSceneMarkedFrame;
+struct FSequencerScriptingRange;
+struct FTimecode;
+template <typename T> class TSubclassOf;
+
 /**
  * Function library containing methods that should be hoisted onto UMovieSceneSequences for scripting purposes
  */
@@ -631,4 +635,12 @@ public:
 	 */
 	UFUNCTION(BlueprintCallable, Category = "Sequencer|Sequence", meta = (ScriptMethod))
 	static bool IsReadOnly(UMovieSceneSequence* Sequence);
-};
\ No newline at end of file
+};
+
+#if UE_ENABLE_INCLUDE_ORDER_DEPRECATED_IN_5_2
+#include "MovieScene.h" // only for FMovieSceneMarkedFrame in the .gen.cpp
+#include "MovieSceneBindingProxy.h"
+#include "MovieSceneTrack.h"
+#include "SequencerScriptingRange.h"
+#include "Templates/SubclassOf.h"
+#endif
diff --git a/Engine/Plugins/MovieScene/SequencerScripting/Source/SequencerScripting/Public/ExtensionLibraries/MovieSceneTrackExtensions.h b/Engine/Plugins/MovieScene/SequencerScripting/Source/SequencerScripting/Public/ExtensionLibraries/MovieSceneTrackExtensions.h
--- a/Engine/Plugins/MovieScene/SequencerScripting/Source/SequencerScripting/Public/ExtensionLibraries/MovieSceneTrackExtensions.h
+++ b/Engine/Plugins/MovieScene/SequencerScripting/Source/SequencerScripting/Public/ExtensionLibraries/MovieSceneTrackExtensions.h
@@ -2,8 +2,6 @@
 
 #pragma once
 
-#include "CoreTypes.h"
-#include "UObject/ObjectMacros.h"
 #include "Kismet/BlueprintFunctionLibrary.h"
 
 #include "MovieSceneTrackExtensions.generated.h"
@@ -144,4 +142,4 @@ public:
 	 */
 	UFUNCTION(BlueprintCallable, Category = "Sequencer|Track", meta=(ScriptMethod, DevelopmentOnly))
 	static void SetSectionToKey(UMovieSceneTrack* Track, UMovieSceneSection* Section); 
-};
\ No newline at end of file
+};
diff --git a/Engine/Plugins/MovieScene/SequencerScripting/Source/SequencerScripting/Public/ExtensionLibraries/SequencerScriptingRangeExtensions.h b/Engine/Plugins/MovieScene/SequencerScripting/Source/SequencerScripting/Public/ExtensionLibraries/SequencerScriptingRangeExtensions.h
--- a/Engine/Plugins/MovieScene/SequencerScripting/Source/SequencerScripting/Public/ExtensionLibraries/SequencerScriptingRangeExtensions.h
+++ b/Engine/Plugins/MovieScene/SequencerScripting/Source/SequencerScripting/Public/ExtensionLibraries/SequencerScriptingRangeExtensions.h
@@ -2,13 +2,12 @@
 
 #pragma once
 
-#include "CoreTypes.h"
-#include "UObject/ObjectMacros.h"
 #include "Kismet/BlueprintFunctionLibrary.h"
-#include "SequencerScriptingRange.h"
 
 #include "SequencerScriptingRangeExtensions.generated.h"
 
+struct FSequencerScriptingRange;
+
 /**
  * Function library containing methods that should be hoisted onto FSequencerScriptingRanges
  */
@@ -112,4 +111,8 @@ public:
 	 */
 	UFUNCTION(BlueprintCallable, Category = "Sequencer|Sequence", meta=(ScriptMethod))
 	static void SetEndFrame(UPARAM(ref) FSequencerScriptingRange& Range, int32 End);
-};
\ No newline at end of file
+};
+
+#if UE_ENABLE_INCLUDE_ORDER_DEPRECATED_IN_5_2
+#include "SequencerScriptingRange.h"
+#endif
diff --git a/Engine/Plugins/MovieScene/SequencerScripting/Source/SequencerScripting/Public/SequenceTimeUnit.h b/Engine/Plugins/MovieScene/SequencerScripting/Source/SequencerScripting/Public/SequenceTimeUnit.h
--- a/Engine/Plugins/MovieScene/SequencerScripting/Source/SequencerScripting/Public/SequenceTimeUnit.h
+++ b/Engine/Plugins/MovieScene/SequencerScripting/Source/SequencerScripting/Public/SequenceTimeUnit.h
@@ -2,8 +2,6 @@
 
 #pragma once
 
-#include "CoreMinimal.h"
-#include "UObject/ObjectMacros.h"
 #include "SequenceTimeUnit.generated.h"
 
 /**
@@ -17,4 +15,8 @@ enum class ESequenceTimeUnit : uint8
 	DisplayRate,
 	/** Tick Resolution is the internal resolution that data is actually stored in, such as 24000 giving you 24,000 frames per second. This is the smallest interval that data can be stored on. */
 	TickResolution
-};
\ No newline at end of file
+};
+
+#if UE_ENABLE_INCLUDE_ORDER_DEPRECATED_IN_5_2
+#include "CoreMinimal.h"
+#endif
diff --git a/Engine/Plugins/MovieScene/SequencerScripting/Source/SequencerScripting/Public/SequencerBindingProxy.h b/Engine/Plugins/MovieScene/SequencerScripting/Source/SequencerScripting/Public/SequencerBindingProxy.h
--- a/Engine/Plugins/MovieScene/SequencerScripting/Source/SequencerScripting/Public/SequencerBindingProxy.h
+++ b/Engine/Plugins/MovieScene/SequencerScripting/Source/SequencerScripting/Public/SequencerBindingProxy.h
@@ -3,8 +3,13 @@
 #pragma once
 
 #include "CoreTypes.h"
-#include "UObject/ObjectMacros.h"
-#include "MovieSceneBindingProxy.h"
+
+struct FMovieSceneBindingProxy;
 
 UE_DEPRECATED(5.1, "Use FMovieSceneBindingProxy instead.") 
-typedef FMovieSceneBindingProxy FSequencerBindingProxy;
\ No newline at end of file
+typedef FMovieSceneBindingProxy FSequencerBindingProxy;
+
+#if UE_ENABLE_INCLUDE_ORDER_DEPRECATED_IN_5_2
+#include "MovieSceneBindingProxy.h"
+#include "UObject/ObjectMacros.h"
+#endif
diff --git a/Engine/Plugins/MovieScene/SequencerScripting/Source/SequencerScripting/Public/SequencerScripting.h b/Engine/Plugins/MovieScene/SequencerScripting/Source/SequencerScripting/Public/SequencerScripting.h
--- a/Engine/Plugins/MovieScene/SequencerScripting/Source/SequencerScripting/Public/SequencerScripting.h
+++ b/Engine/Plugins/MovieScene/SequencerScripting/Source/SequencerScripting/Public/SequencerScripting.h
@@ -2,8 +2,7 @@
 
 #pragma once
 
-#include "CoreMinimal.h"
-#include "Modules/ModuleManager.h"
+#include "Modules/ModuleInterface.h"
 
 class FSequencerScriptingModule : public IModuleInterface
 {
@@ -13,3 +12,8 @@ public:
 	virtual void StartupModule() override;
 	virtual void ShutdownModule() override;
 };
+
+#if UE_ENABLE_INCLUDE_ORDER_DEPRECATED_IN_5_2
+#include "CoreMinimal.h"
+#include "Modules/ModuleManager.h"
+#endif
diff --git a/Engine/Plugins/MovieScene/SequencerScripting/Source/SequencerScripting/Public/SequencerScriptingRange.h b/Engine/Plugins/MovieScene/SequencerScripting/Source/SequencerScripting/Public/SequencerScriptingRange.h
--- a/Engine/Plugins/MovieScene/SequencerScripting/Source/SequencerScripting/Public/SequencerScriptingRange.h
+++ b/Engine/Plugins/MovieScene/SequencerScripting/Source/SequencerScripting/Public/SequencerScriptingRange.h
@@ -2,11 +2,8 @@
 
 #pragma once
 
-#include "CoreTypes.h"
-#include "UObject/ObjectMacros.h"
 #include "MovieSceneTimeHelpers.h"
 #include "Misc/FrameRate.h"
-#include "Math/Range.h"
 
 #include "SequencerScriptingRange.generated.h"
 
@@ -97,4 +94,4 @@ struct FSequencerScriptingRange
 
 	UPROPERTY()
 	FFrameRate InternalRate;
-};
\ No newline at end of file
+};
diff --git a/Engine/Plugins/MovieScene/SequencerScripting/Source/SequencerScriptingEditor/Public/SequencerTools.h b/Engine/Plugins/MovieScene/SequencerScripting/Source/SequencerScriptingEditor/Public/SequencerTools.h
--- a/Engine/Plugins/MovieScene/SequencerScripting/Source/SequencerScriptingEditor/Public/SequencerTools.h
+++ b/Engine/Plugins/MovieScene/SequencerScripting/Source/SequencerScriptingEditor/Public/SequencerTools.h
@@ -2,16 +2,18 @@
 
 #pragma once
 
-#include "CoreMinimal.h"
 #include "Kismet/BlueprintFunctionLibrary.h"
-#include "LevelSequence.h"
-#include "Modules/ModuleManager.h"
 #include "MovieSceneCaptureDialogModule.h"
-#include "Channels/MovieSceneEvent.h"
 #include "MovieSceneBindingProxy.h"
-#include "SequencerScriptingRange.h"
+#include "MovieSceneTrack.h"
 #include "SequencerTools.generated.h"
 
+class IModuleInterface;
+class ULevelSequence;
+class UMovieSceneUserImportFBXSettings;
+struct FMovieSceneEvent;
+struct FSequencerScriptingRange;
+
 class UFbxExportOption;
 class UAnimSequenceExportOption;
 class UAnimSequence;
@@ -275,4 +277,11 @@ public:
 	static bool ImportFBXToControlRig(UWorld* World, ULevelSequence* InSequence, const FString& ActorWithControlRigTrack, const TArray<FString>& SelectedControlRigNames,
 		UMovieSceneUserImportFBXControlRigSettings* ImportFBXControlRigSettings,
 		const FString& ImportFilename);
-};
\ No newline at end of file
+};
+
+#if UE_ENABLE_INCLUDE_ORDER_DEPRECATED_IN_5_2
+#include "Channels/MovieSceneEvent.h"
+#include "CoreMinimal.h"
+#include "LevelSequence.h"
+#include "SequencerScriptingRange.h"
+#endif
diff --git a/Engine/Plugins/MovieScene/TemplateSequence/Source/TemplateSequence/Public/CameraAnimationSequencePlayer.h b/Engine/Plugins/MovieScene/TemplateSequence/Source/TemplateSequence/Public/CameraAnimationSequencePlayer.h
--- a/Engine/Plugins/MovieScene/TemplateSequence/Source/TemplateSequence/Public/CameraAnimationSequencePlayer.h
+++ b/Engine/Plugins/MovieScene/TemplateSequence/Source/TemplateSequence/Public/CameraAnimationSequencePlayer.h
@@ -2,11 +2,12 @@
 
 #pragma once
 
-#include "CineCameraComponent.h"
-#include "Engine/Scene.h"
+#include "CineCameraSettings.h"
 #include "TemplateSequencePlayer.h"
 #include "CameraAnimationSequencePlayer.generated.h"
 
+struct FMinimalViewInfo;
+
 /**
  * A dummy class that we give to a sequence in lieu of an actual camera actor.
  */
@@ -189,3 +190,7 @@ private:
 	TEnumAsByte<EMovieScenePlayerStatus::Type> Status;
 };
 
+
+#if UE_ENABLE_INCLUDE_ORDER_DEPRECATED_IN_5_2
+#include "CineCameraComponent.h"
+#endif
diff --git a/Engine/Plugins/MovieScene/TemplateSequence/Source/TemplateSequence/Public/Sections/TemplateSequenceSection.h b/Engine/Plugins/MovieScene/TemplateSequence/Source/TemplateSequence/Public/Sections/TemplateSequenceSection.h
--- a/Engine/Plugins/MovieScene/TemplateSequence/Source/TemplateSequence/Public/Sections/TemplateSequenceSection.h
+++ b/Engine/Plugins/MovieScene/TemplateSequence/Source/TemplateSequence/Public/Sections/TemplateSequenceSection.h
@@ -2,16 +2,13 @@
 
 #pragma once
 
-#include "CoreMinimal.h"
-#include "UObject/ObjectMacros.h"
 #include "Channels/MovieSceneFloatChannel.h"
-#include "EntitySystem/BuiltInComponentTypes.h"
-#include "EntitySystem/IMovieSceneEntityProvider.h"
-#include "EntitySystem/MovieSceneEntityIDs.h"
 #include "EntitySystem/MovieScenePropertyBinding.h"
 #include "Sections/MovieSceneSubSection.h"
 #include "TemplateSequenceSection.generated.h"
 
+enum class EMovieSceneChannelProxyType : uint8;
+
 /**
  * Defines the type of property scaling for a template sequence.
  */
@@ -79,3 +76,8 @@ public:
 	UPROPERTY()
 	TArray<FTemplateSectionPropertyScale> PropertyScales;
 };
+
+#if UE_ENABLE_INCLUDE_ORDER_DEPRECATED_IN_5_2
+#include "CoreMinimal.h"
+#include "EntitySystem/BuiltInComponentTypes.h"
+#endif
diff --git a/Engine/Plugins/MovieScene/TemplateSequence/Source/TemplateSequence/Public/SequenceCameraShake.h b/Engine/Plugins/MovieScene/TemplateSequence/Source/TemplateSequence/Public/SequenceCameraShake.h
--- a/Engine/Plugins/MovieScene/TemplateSequence/Source/TemplateSequence/Public/SequenceCameraShake.h
+++ b/Engine/Plugins/MovieScene/TemplateSequence/Source/TemplateSequence/Public/SequenceCameraShake.h
@@ -3,12 +3,10 @@
 #pragma once
 
 #include "Camera/CameraShakeBase.h"
-#include "CineCameraComponent.h"
-#include "CameraAnimationSequence.h"
-#include "EntitySystem/MovieSceneEntityIDs.h"
-#include "UObject/WeakObjectPtr.h"
 #include "SequenceCameraShake.generated.h"
 
+struct FMinimalViewInfo;
+
 class UCameraAnimationSequenceCameraStandIn;
 class UMovieSceneEntitySystemLinker;
 class UMovieSceneSequence;
@@ -83,3 +81,9 @@ private:
 	UPROPERTY(Instanced, Transient)
 	TObjectPtr<UCameraAnimationSequenceCameraStandIn> CameraStandIn;
 };
+
+#if UE_ENABLE_INCLUDE_ORDER_DEPRECATED_IN_5_2
+#include "CameraAnimationSequence.h"
+#include "CineCameraComponent.h"
+#include "EntitySystem/MovieSceneEntityIDs.h"
+#endif
diff --git a/Engine/Plugins/MovieScene/TemplateSequence/Source/TemplateSequence/Public/TemplateSequence.h b/Engine/Plugins/MovieScene/TemplateSequence/Source/TemplateSequence/Public/TemplateSequence.h
--- a/Engine/Plugins/MovieScene/TemplateSequence/Source/TemplateSequence/Public/TemplateSequence.h
+++ b/Engine/Plugins/MovieScene/TemplateSequence/Source/TemplateSequence/Public/TemplateSequence.h
@@ -3,8 +3,6 @@
 #pragma once
 
 #include "MovieSceneSequence.h"
-#include "MovieScene.h"
-#include "UObject/SoftObjectPtr.h"
 #include "TemplateSequence.generated.h"
 
 class MovieSceneTrack;
@@ -69,3 +67,7 @@ public:
 	UPROPERTY()
 	TMap<FGuid, FName> BoundActorComponents;
 };
+
+#if UE_ENABLE_INCLUDE_ORDER_DEPRECATED_IN_5_2
+#include "MovieScene.h"
+#endif
diff --git a/Engine/Plugins/MovieScene/TemplateSequence/Source/TemplateSequence/Public/TemplateSequenceActor.h b/Engine/Plugins/MovieScene/TemplateSequence/Source/TemplateSequence/Public/TemplateSequenceActor.h
--- a/Engine/Plugins/MovieScene/TemplateSequence/Source/TemplateSequence/Public/TemplateSequenceActor.h
+++ b/Engine/Plugins/MovieScene/TemplateSequence/Source/TemplateSequence/Public/TemplateSequenceActor.h
@@ -2,16 +2,14 @@
 
 #pragma once
 
-#include "CoreMinimal.h"
-#include "TemplateSequencePlayer.h"
 #include "IMovieScenePlaybackClient.h"
-#include "MovieSceneBindingOwnerInterface.h"
 #include "GameFramework/Actor.h"
-#include "UObject/ObjectMacros.h"
-#include "UObject/Object.h"
-#include "UObject/SoftObjectPath.h"
+#include "MovieSceneSequencePlaybackSettings.h"
 #include "TemplateSequenceActor.generated.h"
 
+class UTemplateSequencePlayer;
+namespace EEndPlayReason { enum Type : int; }
+
 class UTemplateSequence;
 
 /**
@@ -134,3 +132,9 @@ private:
 
 	void OnSequenceLoaded(const FName& PackageName, UPackage* Package, EAsyncLoadingResult::Type Result);
 };
+
+#if UE_ENABLE_INCLUDE_ORDER_DEPRECATED_IN_5_2
+#include "CoreMinimal.h"
+#include "MovieSceneBindingOwnerInterface.h"
+#include "TemplateSequencePlayer.h"
+#endif
diff --git a/Engine/Plugins/MovieScene/TemplateSequence/Source/TemplateSequence/Public/TemplateSequenceComponentTypes.h b/Engine/Plugins/MovieScene/TemplateSequence/Source/TemplateSequence/Public/TemplateSequenceComponentTypes.h
--- a/Engine/Plugins/MovieScene/TemplateSequence/Source/TemplateSequence/Public/TemplateSequenceComponentTypes.h
+++ b/Engine/Plugins/MovieScene/TemplateSequence/Source/TemplateSequence/Public/TemplateSequenceComponentTypes.h
@@ -2,11 +2,10 @@
 
 #pragma once
 
-#include "CoreMinimal.h"
-#include "EntitySystem/MovieSceneEntitySystemTypes.h"
 #include "EntitySystem/MovieScenePropertyBinding.h"
 #include "Evaluation/MovieSceneEvaluationOperand.h"
-#include "Sections/TemplateSequenceSection.h"
+
+enum class ETemplateSectionPropertyScaleType;
 
 struct FMovieSceneFloatChannel;
 
@@ -55,3 +54,8 @@ private:
 } // namespace MovieScene
 } // namespace UE
 
+
+#if UE_ENABLE_INCLUDE_ORDER_DEPRECATED_IN_5_2
+#include "CoreMinimal.h"
+#include "Sections/TemplateSequenceSection.h"
+#endif
diff --git a/Engine/Plugins/MovieScene/TemplateSequence/Source/TemplateSequence/Public/TemplateSequencePlayer.h b/Engine/Plugins/MovieScene/TemplateSequence/Source/TemplateSequence/Public/TemplateSequencePlayer.h
--- a/Engine/Plugins/MovieScene/TemplateSequence/Source/TemplateSequence/Public/TemplateSequencePlayer.h
+++ b/Engine/Plugins/MovieScene/TemplateSequence/Source/TemplateSequence/Public/TemplateSequencePlayer.h
@@ -5,6 +5,8 @@
 #include "MovieSceneSequencePlayer.h"
 #include "TemplateSequencePlayer.generated.h"
 
+struct FMovieSceneSequencePlaybackSettings;
+
 class ATemplateSequenceActor;
 class UTemplateSequence;
 
diff --git a/Engine/Plugins/MovieScene/TemplateSequence/Source/TemplateSequence/Public/Tracks/TemplateSequenceTrack.h b/Engine/Plugins/MovieScene/TemplateSequence/Source/TemplateSequence/Public/Tracks/TemplateSequenceTrack.h
--- a/Engine/Plugins/MovieScene/TemplateSequence/Source/TemplateSequence/Public/Tracks/TemplateSequenceTrack.h
+++ b/Engine/Plugins/MovieScene/TemplateSequence/Source/TemplateSequence/Public/Tracks/TemplateSequenceTrack.h
@@ -2,10 +2,6 @@
 
 #pragma once
 
-#include "CoreMinimal.h"
-#include "UObject/ObjectMacros.h"
-#include "Misc/InlineValue.h"
-#include "Compilation/MovieSceneSegmentCompiler.h"
 #include "Tracks/MovieSceneSubTrack.h"
 #include "TemplateSequenceTrack.generated.h"
 
@@ -35,3 +31,7 @@ public:
 private:
 	static uint16 GetEvaluationPriority() { return uint16(0x1000); }  // One more than spawn tracks.
 };
+
+#if UE_ENABLE_INCLUDE_ORDER_DEPRECATED_IN_5_2
+#include "CoreMinimal.h"
+#endif
diff --git a/Engine/Plugins/Online/Android/OnlineSubsystemGooglePlay/Source/Private/OnlineAchievementsInterfaceGooglePlay.h b/Engine/Plugins/Online/Android/OnlineSubsystemGooglePlay/Source/Private/OnlineAchievementsInterfaceGooglePlay.h
--- a/Engine/Plugins/Online/Android/OnlineSubsystemGooglePlay/Source/Private/OnlineAchievementsInterfaceGooglePlay.h
+++ b/Engine/Plugins/Online/Android/OnlineSubsystemGooglePlay/Source/Private/OnlineAchievementsInterfaceGooglePlay.h
@@ -6,6 +6,7 @@
 #include "Interfaces/OnlineAchievementsInterface.h"
 #include "OnlineSubsystemTypes.h"
 #include "AndroidRuntimeSettings.h"
+#include "OnlineSubsystemGooglePlayPackage.h"
 
 THIRD_PARTY_INCLUDES_START
 #include <jni.h>
diff --git a/Engine/Plugins/Online/Android/OnlineSubsystemGooglePlay/Source/Private/OnlineSubsystemGooglePlay.cpp b/Engine/Plugins/Online/Android/OnlineSubsystemGooglePlay/Source/Private/OnlineSubsystemGooglePlay.cpp
--- a/Engine/Plugins/Online/Android/OnlineSubsystemGooglePlay/Source/Private/OnlineSubsystemGooglePlay.cpp
+++ b/Engine/Plugins/Online/Android/OnlineSubsystemGooglePlay/Source/Private/OnlineSubsystemGooglePlay.cpp
@@ -3,6 +3,7 @@
 //Google Play Services
 
 #include "OnlineSubsystemGooglePlay.h"
+#include "OnlineIdentityErrors.h"
 #include "OnlineIdentityInterfaceGooglePlay.h"
 #include "OnlineStoreGooglePlayCommon.h"
 #include "OnlineAchievementsInterfaceGooglePlay.h"
diff --git a/Engine/Plugins/Online/OnlineFramework/Source/Hotfix/Public/OnlineHotfixManager.h b/Engine/Plugins/Online/OnlineFramework/Source/Hotfix/Public/OnlineHotfixManager.h
--- a/Engine/Plugins/Online/OnlineFramework/Source/Hotfix/Public/OnlineHotfixManager.h
+++ b/Engine/Plugins/Online/OnlineFramework/Source/Hotfix/Public/OnlineHotfixManager.h
@@ -2,17 +2,12 @@
 
 #pragma once
 
-#include "CoreMinimal.h"
-#include "UObject/ObjectMacros.h"
-#include "UObject/Object.h"
-#include "UObject/ScriptMacros.h"
-#include "OnlineSubsystemTypes.h"
-#include "OnlineDelegateMacros.h"
-#include "Misc/Paths.h"
 #include "Misc/ConfigCacheIni.h"
 #include "Interfaces/OnlineTitleFileInterface.h"
 #include "OnlineHotfixManager.generated.h"
 
+struct FCloudFileHeader;
+
 HOTFIX_API DECLARE_LOG_CATEGORY_EXTERN(LogHotfixManager, Display, All);
 
 class FAsyncLoadingFlushContext;
@@ -371,3 +366,7 @@ public:
 	/** Factory method that returns the configured hotfix manager */
 	static UOnlineHotfixManager* Get(UWorld* World);
 };
+
+#if UE_ENABLE_INCLUDE_ORDER_DEPRECATED_IN_5_2
+#include "CoreMinimal.h"
+#endif
diff --git a/Engine/Plugins/Online/OnlineFramework/Source/Hotfix/Public/UpdateManager.h b/Engine/Plugins/Online/OnlineFramework/Source/Hotfix/Public/UpdateManager.h
--- a/Engine/Plugins/Online/OnlineFramework/Source/Hotfix/Public/UpdateManager.h
+++ b/Engine/Plugins/Online/OnlineFramework/Source/Hotfix/Public/UpdateManager.h
@@ -2,17 +2,14 @@
 
 #pragma once
 
-#include "CoreMinimal.h"
-#include "UObject/ObjectMacros.h"
-#include "UObject/Object.h"
-#include "Engine/EngineTypes.h"
-#include "Templates/Casts.h"
-#include "Interfaces/OnlineIdentityInterface.h"
-#include "PatchCheck.h"
-#include "InstallBundleTypes.h"
 #include "Containers/Ticker.h"
+#include "Misc/DateTime.h"
 #include "UpdateManager.generated.h"
 
+enum class EInstallBundleManagerPatchCheckResult : uint32;
+enum class EPatchCheckResult : uint8;
+struct FOnlineError;
+
 class Error;
 class UGameInstance;
 enum class EHotfixResult : uint8;
@@ -490,3 +487,11 @@ protected:
 	/** @return a pointer to the game instance */
 	UGameInstance* GetGameInstance() const;
 };
+
+#if UE_ENABLE_INCLUDE_ORDER_DEPRECATED_IN_5_2
+#include "CoreMinimal.h"
+#include "Engine/EngineTypes.h"
+#include "InstallBundleTypes.h"
+#include "Interfaces/OnlineIdentityInterface.h"
+#include "PatchCheck.h"
+#endif
diff --git a/Engine/Plugins/Online/OnlineFramework/Source/Lobby/Public/LobbyBeaconClient.h b/Engine/Plugins/Online/OnlineFramework/Source/Lobby/Public/LobbyBeaconClient.h
--- a/Engine/Plugins/Online/OnlineFramework/Source/Lobby/Public/LobbyBeaconClient.h
+++ b/Engine/Plugins/Online/OnlineFramework/Source/Lobby/Public/LobbyBeaconClient.h
@@ -2,15 +2,13 @@
 
 #pragma once
 
-#include "CoreMinimal.h"
-#include "UObject/ObjectMacros.h"
-#include "GameFramework/OnlineReplStructs.h"
-#include "UObject/TextProperty.h"
 #include "OnlineBeaconClient.h"
 
 
 #include "LobbyBeaconClient.generated.h"
 
+struct FJoinabilitySettings;
+
 class ALobbyBeaconPlayerState;
 class ALobbyBeaconState;
 class FOnlineSessionSearchResult;
@@ -256,3 +254,9 @@ protected:
 
 	friend class ALobbyBeaconHost;
 };
+
+#if UE_ENABLE_INCLUDE_ORDER_DEPRECATED_IN_5_2
+#include "CoreMinimal.h"
+#include "GameFramework/OnlineReplStructs.h"
+#include "UObject/TextProperty.h"
+#endif
diff --git a/Engine/Plugins/Online/OnlineFramework/Source/Lobby/Public/LobbyBeaconState.h b/Engine/Plugins/Online/OnlineFramework/Source/Lobby/Public/LobbyBeaconState.h
--- a/Engine/Plugins/Online/OnlineFramework/Source/Lobby/Public/LobbyBeaconState.h
+++ b/Engine/Plugins/Online/OnlineFramework/Source/Lobby/Public/LobbyBeaconState.h
@@ -2,17 +2,12 @@
 
 #pragma once
 
-#include "CoreMinimal.h"
 #include "Net/Serialization/FastArraySerializer.h"
-#include "UObject/ObjectMacros.h"
-#include "UObject/Class.h"
-#include "Templates/SubclassOf.h"
-#include "Engine/NetSerialization.h"
-#include "Engine/EngineTypes.h"
 #include "GameFramework/Info.h"
-#include "LobbyBeaconHost.h"
 #include "LobbyBeaconState.generated.h"
 
+class ALobbyBeaconHost;
+
 class ALobbyBeaconPlayerState;
 class ALobbyBeaconState;
 class AOnlineBeaconClient;
@@ -357,3 +352,8 @@ protected:
 	friend ALobbyBeaconHost;
 	friend FLobbyPlayerStateInfoArray;
 };
+
+#if UE_ENABLE_INCLUDE_ORDER_DEPRECATED_IN_5_2
+#include "CoreMinimal.h"
+#include "LobbyBeaconHost.h"
+#endif
diff --git a/Engine/Plugins/Online/OnlineFramework/Source/LoginFlow/Public/ILoginFlowModule.h b/Engine/Plugins/Online/OnlineFramework/Source/LoginFlow/Public/ILoginFlowModule.h
--- a/Engine/Plugins/Online/OnlineFramework/Source/LoginFlow/Public/ILoginFlowModule.h
+++ b/Engine/Plugins/Online/OnlineFramework/Source/LoginFlow/Public/ILoginFlowModule.h
@@ -1,7 +1,8 @@
 // Copyright Epic Games, Inc. All Rights Reserved.
 #pragma once
 #include "Modules/ModuleManager.h"
-#include "ILoginFlowManager.h"
+
+class ILoginFlowManager;
 
 class IOnlineSubsystem;
 class ISlateStyle;
@@ -110,3 +111,7 @@ public:
 	/** Virtual destructor. */
 	virtual ~ILoginFlowModule() { }
 };
+
+#if UE_ENABLE_INCLUDE_ORDER_DEPRECATED_IN_5_2
+#include "ILoginFlowManager.h"
+#endif
diff --git a/Engine/Plugins/Online/OnlineFramework/Source/LoginFlow/Public/LoginFlow.h b/Engine/Plugins/Online/OnlineFramework/Source/LoginFlow/Public/LoginFlow.h
--- a/Engine/Plugins/Online/OnlineFramework/Source/LoginFlow/Public/LoginFlow.h
+++ b/Engine/Plugins/Online/OnlineFramework/Source/LoginFlow/Public/LoginFlow.h
@@ -5,14 +5,17 @@
 /* Dependencies
  *****************************************************************************/
 
-#include "CoreMinimal.h"
-#include "Modules/ModuleManager.h"
-#include "Styling/SlateStyle.h"
-#include "OnlineSubsystem.h"
 
 /* Interfaces
  *****************************************************************************/
 
-#include "ILoginFlowModule.h"
-#include "ILoginFlowManager.h"
 
+
+#if UE_ENABLE_INCLUDE_ORDER_DEPRECATED_IN_5_2
+#include "CoreMinimal.h"
+#include "ILoginFlowManager.h"
+#include "ILoginFlowModule.h"
+#include "Modules/ModuleManager.h"
+#include "OnlineSubsystem.h"
+#include "Styling/SlateStyle.h"
+#endif
diff --git a/Engine/Plugins/Online/OnlineFramework/Source/LoginFlow/Public/LoginFlowManager.h b/Engine/Plugins/Online/OnlineFramework/Source/LoginFlow/Public/LoginFlowManager.h
--- a/Engine/Plugins/Online/OnlineFramework/Source/LoginFlow/Public/LoginFlowManager.h
+++ b/Engine/Plugins/Online/OnlineFramework/Source/LoginFlow/Public/LoginFlowManager.h
@@ -1,10 +1,12 @@
 // Copyright Epic Games, Inc. All Rights Reserved.
 #pragma once
 
-#include "ILoginFlowModule.h"
 #include "ILoginFlowManager.h"
 #include "Interfaces/OnlineExternalUIInterface.h"
 
+enum class ELoginFlowErrorResult : uint8;
+struct FBrowserContextSettings;
+
 class FOnlineSubsystemMcp;
 class SWidget;
 struct FLoginFlowProperties;
@@ -110,3 +112,7 @@ private:
 	TUniquePtr<FAccountCreationFlowProperties> PendingAccountCreation;
 };
 
+
+#if UE_ENABLE_INCLUDE_ORDER_DEPRECATED_IN_5_2
+#include "ILoginFlowModule.h"
+#endif
diff --git a/Engine/Plugins/Online/OnlineFramework/Source/Party/Private/Interactions/CoreInteractions.cpp b/Engine/Plugins/Online/OnlineFramework/Source/Party/Private/Interactions/CoreInteractions.cpp
--- a/Engine/Plugins/Online/OnlineFramework/Source/Party/Private/Interactions/CoreInteractions.cpp
+++ b/Engine/Plugins/Online/OnlineFramework/Source/Party/Private/Interactions/CoreInteractions.cpp
@@ -2,7 +2,6 @@
 
 #include "Interactions/CoreInteractions.h"
 
-#include "OnlineSubsystemNames.h"
 #include "SocialManager.h"
 #include "Chat/SocialChatManager.h"
 
diff --git a/Engine/Plugins/Online/OnlineFramework/Source/Party/Private/Party/PartyPlatformSessionMonitor.cpp b/Engine/Plugins/Online/OnlineFramework/Source/Party/Private/Party/PartyPlatformSessionMonitor.cpp
--- a/Engine/Plugins/Online/OnlineFramework/Source/Party/Private/Party/PartyPlatformSessionMonitor.cpp
+++ b/Engine/Plugins/Online/OnlineFramework/Source/Party/Private/Party/PartyPlatformSessionMonitor.cpp
@@ -6,7 +6,6 @@
 #include "Party/SocialParty.h"
 #include "Online/OnlineSessionNames.h"
 #include "Party/PartyMember.h"
-#include "OnlineSubsystemNames.h"
 #include "SocialToolkit.h"
 #include "SocialManager.h"
 
@@ -1307,4 +1306,4 @@ void FPartyPlatformSessionMonitor::HandleSessionFailure(const FUniqueNetId& Loca
 
 		RetryTickerHandle = FTSTicker::GetCoreTicker().AddTicker(FTickerDelegate::CreateSP(this, &FPartyPlatformSessionMonitor::HandleRetryEstablishingSession), EstablishSessionRetryDelay);
 	}
-}
\ No newline at end of file
+}
diff --git a/Engine/Plugins/Online/OnlineFramework/Source/Party/Private/SocialDebugTools.cpp b/Engine/Plugins/Online/OnlineFramework/Source/Party/Private/SocialDebugTools.cpp
--- a/Engine/Plugins/Online/OnlineFramework/Source/Party/Private/SocialDebugTools.cpp
+++ b/Engine/Plugins/Online/OnlineFramework/Source/Party/Private/SocialDebugTools.cpp
@@ -2,7 +2,6 @@
 
 #include "SocialDebugTools.h"
 
-#include "OnlineSubsystemNames.h"
 #include "OnlineSubsystemUtils.h"
 #include "Interfaces/OnlineUserInterface.h"
 #include "Interfaces/OnlinePresenceInterface.h"
diff --git a/Engine/Plugins/Online/OnlineFramework/Source/Party/Public/Chat/ChatSlashCommands.h b/Engine/Plugins/Online/OnlineFramework/Source/Party/Public/Chat/ChatSlashCommands.h
--- a/Engine/Plugins/Online/OnlineFramework/Source/Party/Public/Chat/ChatSlashCommands.h
+++ b/Engine/Plugins/Online/OnlineFramework/Source/Party/Public/Chat/ChatSlashCommands.h
@@ -4,14 +4,9 @@
 #pragma once
 //#include "UObject/Interface.h"
 
-#include "UObject/Object.h"
-#include "Templates/SharedPointer.h"
-#include "Internationalization/Text.h"
-#include "Engine/LocalPlayer.h"
 #include "Interactions/SocialInteractionHandle.h"
 
-#include "Internationalization/Internationalization.h"
-#include "Internationalization/CulturePointer.h"
+#include "UObject/WeakObjectPtr.h"
 
 class FSocialUserList;
 class FChatSlashCommand;
@@ -146,4 +141,10 @@ private:
 	* So, there exists the following function to re-query the localization
 	*/
 	void CacheStringDataForLocalization() const;
-};
\ No newline at end of file
+};
+
+#if UE_ENABLE_INCLUDE_ORDER_DEPRECATED_IN_5_2
+#include "Engine/LocalPlayer.h"
+#include "Internationalization/Internationalization.h"
+#include "UObject/Object.h"
+#endif
