Sequencer: Modify the section for setter scripting functions - AddKeyInChannel, RemoveKeyFromChannel, SetDefaultInChannel, ReoveDefaultFromChannel, SetTimeInChannel, SetValueInChannel

Tested calling an editor utility widget with function to add a key within a begin transaction / end transaction correctly undos.

diff --git a/Engine/Plugins/MovieScene/SequencerScripting/Source/SequencerScripting/Private/KeysAndChannels/MovieSceneScriptingActorReference.h b/Engine/Plugins/MovieScene/SequencerScripting/Source/SequencerScripting/Private/KeysAndChannels/MovieSceneScriptingActorReference.h
--- a/Engine/Plugins/MovieScene/SequencerScripting/Source/SequencerScripting/Private/KeysAndChannels/MovieSceneScriptingActorReference.h
+++ b/Engine/Plugins/MovieScene/SequencerScripting/Source/SequencerScripting/Private/KeysAndChannels/MovieSceneScriptingActorReference.h
@@ -36,7 +36,7 @@ public:
 	* @param TimeUnit		Should the NewFrameNumber be interpreted as Display Rate frames or in Tick Resolution?
 	*/
 	UFUNCTION(BlueprintCallable, Category = "Sequencer|Keys", meta = (DisplayName = "Set Time (Actor Reference)"))
-	void SetTime(const FFrameNumber& NewFrameNumber, float SubFrame = 0.f, ESequenceTimeUnit TimeUnit = ESequenceTimeUnit::DisplayRate) { SetTimeInChannel(KeyHandle, OwningSequence, NewFrameNumber, TimeUnit, SubFrame); }
+	void SetTime(const FFrameNumber& NewFrameNumber, float SubFrame = 0.f, ESequenceTimeUnit TimeUnit = ESequenceTimeUnit::DisplayRate) { SetTimeInChannel(KeyHandle, OwningSequence, OwningSection, NewFrameNumber, TimeUnit, SubFrame); }
 
 	/**
 	* Gets the value for this key from the owning channel.
@@ -57,7 +57,7 @@ public:
 	void SetValue(const FMovieSceneObjectBindingID& InNewValue)
 	{
 		FMovieSceneActorReferenceKey ReferenceKey = FMovieSceneActorReferenceKey(InNewValue);
-		SetValueInChannel(KeyHandle, ReferenceKey);
+		SetValueInChannel(KeyHandle, OwningSection, ReferenceKey);
 	}
 };
 
@@ -140,12 +140,25 @@ public:
 	UFUNCTION(BlueprintCallable, Category = "Sequencer|Keys", meta = (DisplayName = "Remove Default (Actor Reference)"))
 	void RemoveDefault()
 	{
+		if (OwningSection.IsValid())
+		{
+			OwningSection->Modify();
+		}
+
 		// FMovieSceneActorReferenceData doesn't implement RemoveDefault, instead it implements ClearDefault(). Wrapping this function by hand,
 		// and not falling back to the template as a result, but keeping the same function name so it is consistent with the other scripting channels.
 		FMovieSceneActorReferenceData* Channel = ChannelHandle.Get();
 		if (Channel)
 		{
 			Channel->ClearDefault();
+
+#if WITH_EDITOR
+			const FMovieSceneChannelMetaData* MetaData = ChannelHandle.GetMetaData();
+			if (MetaData && OwningSection.IsValid() && OwningSequence.IsValid() && OwningSequence->GetMovieScene())
+			{
+				OwningSequence->GetMovieScene()->OnChannelChanged().Broadcast(MetaData, OwningSection.Get());
+			}
+#endif
 			return;
 		}
 		FFrame::KismetExecutionMessage(TEXT("Invalid ChannelHandle for MovieSceneScriptingChannel, failed to remove default value."), ELogVerbosity::Error);
diff --git a/Engine/Plugins/MovieScene/SequencerScripting/Source/SequencerScripting/Private/KeysAndChannels/MovieSceneScriptingBool.h b/Engine/Plugins/MovieScene/SequencerScripting/Source/SequencerScripting/Private/KeysAndChannels/MovieSceneScriptingBool.h
--- a/Engine/Plugins/MovieScene/SequencerScripting/Source/SequencerScripting/Private/KeysAndChannels/MovieSceneScriptingBool.h
+++ b/Engine/Plugins/MovieScene/SequencerScripting/Source/SequencerScripting/Private/KeysAndChannels/MovieSceneScriptingBool.h
@@ -35,7 +35,7 @@ public:
 	* @param TimeUnit		Should the NewFrameNumber be interpreted as Display Rate frames or in Tick Resolution?
 	*/
 	UFUNCTION(BlueprintCallable, Category = "Sequencer|Keys", meta = (DisplayName = "Set Time (Bool)"))
-	void SetTime(const FFrameNumber& NewFrameNumber, float SubFrame = 0.f, ESequenceTimeUnit TimeUnit = ESequenceTimeUnit::DisplayRate) { SetTimeInChannel(KeyHandle, OwningSequence, NewFrameNumber, TimeUnit, SubFrame); }
+	void SetTime(const FFrameNumber& NewFrameNumber, float SubFrame = 0.f, ESequenceTimeUnit TimeUnit = ESequenceTimeUnit::DisplayRate) { SetTimeInChannel(KeyHandle, OwningSequence, OwningSection, NewFrameNumber, TimeUnit, SubFrame); }
 
 	/**
 	* Gets the value for this key from the owning channel.
@@ -54,7 +54,7 @@ public:
 	UFUNCTION(BlueprintCallable, Category = "Sequencer|Keys", meta = (DisplayName = "Set Value (Bool)"))
 	void SetValue(bool InNewValue)
 	{
-		SetValueInChannel(KeyHandle, InNewValue);
+		SetValueInChannel(KeyHandle, OwningSection, InNewValue);
 	}
 };
 
@@ -158,7 +158,7 @@ public:
 	UFUNCTION(BlueprintCallable, Category = "Sequencer|Keys", meta = (DisplayName = "Remove Default (Bool)"))
 	void RemoveDefault()
 	{
-		Impl::RemoveDefaultFromChannel(ChannelHandle);
+		Impl::RemoveDefaultFromChannel(ChannelHandle, OwningSequence, OwningSection);
 	}
 
 	/**
diff --git a/Engine/Plugins/MovieScene/SequencerScripting/Source/SequencerScripting/Private/KeysAndChannels/MovieSceneScriptingByte.h b/Engine/Plugins/MovieScene/SequencerScripting/Source/SequencerScripting/Private/KeysAndChannels/MovieSceneScriptingByte.h
--- a/Engine/Plugins/MovieScene/SequencerScripting/Source/SequencerScripting/Private/KeysAndChannels/MovieSceneScriptingByte.h
+++ b/Engine/Plugins/MovieScene/SequencerScripting/Source/SequencerScripting/Private/KeysAndChannels/MovieSceneScriptingByte.h
@@ -37,7 +37,7 @@ public:
 	* @param TimeUnit		Should the NewFrameNumber be interpreted as Display Rate frames or in Tick Resolution?
 	*/
 	UFUNCTION(BlueprintCallable, Category = "Sequencer|Keys", meta = (DisplayName = "Set Time (Enum)"))
-	void SetTime(const FFrameNumber& NewFrameNumber, float SubFrame = 0.f, ESequenceTimeUnit TimeUnit = ESequenceTimeUnit::DisplayRate) { SetTimeInChannel(KeyHandle, OwningSequence, NewFrameNumber, TimeUnit, SubFrame); }
+	void SetTime(const FFrameNumber& NewFrameNumber, float SubFrame = 0.f, ESequenceTimeUnit TimeUnit = ESequenceTimeUnit::DisplayRate) { SetTimeInChannel(KeyHandle, OwningSequence, OwningSection, NewFrameNumber, TimeUnit, SubFrame); }
 
 	/**
 	* Gets the value for this key from the owning channel.
@@ -56,7 +56,7 @@ public:
 	UFUNCTION(BlueprintCallable, Category = "Sequencer|Keys", meta = (DisplayName = "Set Value (Enum)"))
 	void SetValue(uint8 InNewValue)
 	{
-		SetValueInChannel(KeyHandle, InNewValue);
+		SetValueInChannel(KeyHandle, OwningSection, InNewValue);
 	}
 };
 
@@ -130,7 +130,7 @@ public:
 	UFUNCTION(BlueprintCallable, Category = "Sequencer|Keys", meta = (DisplayName = "Remove Default (Enum)"))
 	void RemoveDefault()
 	{
-		Impl::RemoveDefaultFromChannel(ChannelHandle);
+		Impl::RemoveDefaultFromChannel(ChannelHandle, OwningSequence, OwningSection);
 	}
 
 	/**
diff --git a/Engine/Plugins/MovieScene/SequencerScripting/Source/SequencerScripting/Private/KeysAndChannels/MovieSceneScriptingChannel.h b/Engine/Plugins/MovieScene/SequencerScripting/Source/SequencerScripting/Private/KeysAndChannels/MovieSceneScriptingChannel.h
--- a/Engine/Plugins/MovieScene/SequencerScripting/Source/SequencerScripting/Private/KeysAndChannels/MovieSceneScriptingChannel.h
+++ b/Engine/Plugins/MovieScene/SequencerScripting/Source/SequencerScripting/Private/KeysAndChannels/MovieSceneScriptingChannel.h
@@ -66,6 +66,11 @@ struct TMovieSceneScriptingChannel
 				KeyTime = FFrameRate::TransformTime(KeyTime, UMovieSceneSequenceExtensions::GetDisplayRate(Sequence.Get()), UMovieSceneSequenceExtensions::GetTickResolution(Sequence.Get())).RoundToFrame();
 			}
 
+			if (Section.IsValid())
+			{
+				Section->Modify();
+			}
+
 			using namespace UE::MovieScene;
 			Key->KeyHandle = AddKeyToChannel(Channel, KeyTime, NewValue, Interpolation);
 			Key->ChannelHandle = ChannelHandle;
@@ -76,7 +81,6 @@ struct TMovieSceneScriptingChannel
 			const FMovieSceneChannelMetaData* MetaData = ChannelHandle.GetMetaData();
 			if (MetaData && Section.IsValid() && Sequence.IsValid() && Sequence->GetMovieScene())
 			{
-				Section.Get()->MarkAsChanged();
 				Sequence->GetMovieScene()->OnChannelChanged().Broadcast(MetaData, Section.Get());
 			}
 #endif
@@ -98,13 +102,17 @@ struct TMovieSceneScriptingChannel
 		ChannelType* Channel = ChannelHandle.Get();
 		if (Channel)
 		{
+			if (Key->OwningSection.IsValid())
+			{
+				Key->OwningSection->Modify();
+			}
+
 			TArrayView<FKeyHandle, int> KeyArray = TArrayView<FKeyHandle, int>(&Key->KeyHandle, 1);
 			Channel->DeleteKeys(KeyArray);
 #if WITH_EDITOR
 			const FMovieSceneChannelMetaData* MetaData = ChannelHandle.GetMetaData();
 			if (MetaData && Key->OwningSection.IsValid() && Key->OwningSequence.IsValid() && Key->OwningSequence->GetMovieScene())
 			{
-				Key->OwningSection.Get()->MarkAsChanged();
 				Key->OwningSequence->GetMovieScene()->OnChannelChanged().Broadcast(MetaData, Key->OwningSection.Get());
 			}
 #endif
@@ -196,6 +204,11 @@ struct TMovieSceneScriptingChannel
 		ChannelType* Channel = ChannelHandle.Get();
 		if (Channel)
 		{
+			if (Section.IsValid())
+			{
+				Section->Modify();
+			}
+
 			using namespace UE::MovieScene;
 			SetChannelDefault(Channel, InDefaultValue);
 
@@ -203,7 +216,6 @@ struct TMovieSceneScriptingChannel
 			const FMovieSceneChannelMetaData* MetaData = ChannelHandle.GetMetaData();
 			if (MetaData && Section.IsValid() && Sequence.IsValid() && Sequence->GetMovieScene())
 			{
-				Section.Get()->MarkAsChanged();
 				Sequence->GetMovieScene()->OnChannelChanged().Broadcast(MetaData, Section.Get());
 			}
 #endif
@@ -212,13 +224,26 @@ struct TMovieSceneScriptingChannel
 		FFrame::KismetExecutionMessage(TEXT("Invalid ChannelHandle for MovieSceneScriptingChannel, failed to set default value."), ELogVerbosity::Error);
 	}
 
-	static void RemoveDefaultFromChannel(TMovieSceneChannelHandle<ChannelType> ChannelHandle)
+	static void RemoveDefaultFromChannel(TMovieSceneChannelHandle<ChannelType> ChannelHandle, TWeakObjectPtr<UMovieSceneSequence> Sequence, TWeakObjectPtr<UMovieSceneSection> Section)
 	{
 		ChannelType* Channel = ChannelHandle.Get();
 		if (Channel)
 		{
+			if (Section.IsValid())
+			{
+				Section->Modify();
+			}
+
 			using namespace UE::MovieScene;
 			RemoveChannelDefault(Channel);
+
+#if WITH_EDITOR
+			const FMovieSceneChannelMetaData* MetaData = ChannelHandle.GetMetaData();
+			if (MetaData && Section.IsValid() && Sequence.IsValid() && Sequence->GetMovieScene())
+			{
+				Sequence->GetMovieScene()->OnChannelChanged().Broadcast(MetaData, Section.Get());
+			}
+#endif
 			return;
 		}
 		FFrame::KismetExecutionMessage(TEXT("Invalid ChannelHandle for MovieSceneScriptingChannel, failed to remove default value."), ELogVerbosity::Error);
@@ -281,7 +306,7 @@ struct TMovieSceneScriptingKey
 		return FFrameTime();
 	}
 
-	void SetTimeInChannel(FKeyHandle KeyHandle, TWeakObjectPtr<UMovieSceneSequence> Sequence, const FFrameNumber NewFrameNumber, ESequenceTimeUnit TimeUnit, float SubFrame)
+	void SetTimeInChannel(FKeyHandle KeyHandle, TWeakObjectPtr<UMovieSceneSequence> Sequence, TWeakObjectPtr<UMovieSceneSection> Section, const FFrameNumber NewFrameNumber, ESequenceTimeUnit TimeUnit, float SubFrame)
 	{
 		if (!Sequence.IsValid())
 		{
@@ -310,6 +335,11 @@ struct TMovieSceneScriptingKey
 		ChannelType* Channel = ChannelHandle.Get();
 		if (Channel)
 		{
+			if (Section.IsValid())
+			{
+				Section->Modify();
+			}
+
 			Channel->SetKeyTime(KeyHandle, KeyFrameNumber);
 			return;
 		}
@@ -338,11 +368,16 @@ struct TMovieSceneScriptingKey
 		return Value;
 	}
 
-	void SetValueInChannel(FKeyHandle KeyHandle, ChannelDataType InNewValue)
+	void SetValueInChannel(FKeyHandle KeyHandle, TWeakObjectPtr<UMovieSceneSection> Section, ChannelDataType InNewValue)
 	{
 		ChannelType* Channel = ChannelHandle.Get();
 		if (Channel)
 		{
+			if (Section.IsValid())
+			{
+				Section->Modify();
+			}
+
 			using namespace UE::MovieScene;
 			if (!AssignValue(Channel, KeyHandle, InNewValue))
 			{
diff --git a/Engine/Plugins/MovieScene/SequencerScripting/Source/SequencerScripting/Private/KeysAndChannels/MovieSceneScriptingDouble.h b/Engine/Plugins/MovieScene/SequencerScripting/Source/SequencerScripting/Private/KeysAndChannels/MovieSceneScriptingDouble.h
--- a/Engine/Plugins/MovieScene/SequencerScripting/Source/SequencerScripting/Private/KeysAndChannels/MovieSceneScriptingDouble.h
+++ b/Engine/Plugins/MovieScene/SequencerScripting/Source/SequencerScripting/Private/KeysAndChannels/MovieSceneScriptingDouble.h
@@ -39,7 +39,7 @@ public:
 	* @param TimeUnit		Should the NewFrameNumber be interpreted as Display Rate frames or in Tick Resolution?
 	*/
 	UFUNCTION(BlueprintCallable, Category = "Sequencer|Keys", meta = (DisplayName = "Set Time (Double)"))
-	void SetTime(const FFrameNumber& NewFrameNumber, float SubFrame = 0.f, ESequenceTimeUnit TimeUnit = ESequenceTimeUnit::DisplayRate) { SetTimeInChannel(KeyHandle, OwningSequence, NewFrameNumber, TimeUnit, SubFrame); }
+	void SetTime(const FFrameNumber& NewFrameNumber, float SubFrame = 0.f, ESequenceTimeUnit TimeUnit = ESequenceTimeUnit::DisplayRate) { SetTimeInChannel(KeyHandle, OwningSequence, OwningSection, NewFrameNumber, TimeUnit, SubFrame); }
 
 	/**
 	* Gets the value for this key from the owning channel.
@@ -60,7 +60,7 @@ public:
 	{
 		FMovieSceneDoubleValue ExistValue = GetValueFromChannel(KeyHandle);
 		ExistValue.Value = InNewValue;
-		SetValueInChannel(KeyHandle, ExistValue);
+		SetValueInChannel(KeyHandle, OwningSection, ExistValue);
 	}
 
 	/**
@@ -82,7 +82,7 @@ public:
 	{
 		FMovieSceneDoubleValue ExistValue = GetValueFromChannel(KeyHandle);
 		ExistValue.InterpMode = InNewValue;
-		SetValueInChannel(KeyHandle, ExistValue);
+		SetValueInChannel(KeyHandle, OwningSection, ExistValue);
 	}
 
 	/**
@@ -104,7 +104,7 @@ public:
 	{
 		FMovieSceneDoubleValue ExistValue = GetValueFromChannel(KeyHandle);
 		ExistValue.TangentMode = InNewValue;
-		SetValueInChannel(KeyHandle, ExistValue);
+		SetValueInChannel(KeyHandle, OwningSection, ExistValue);
 	}
 
 	/**
@@ -127,7 +127,7 @@ public:
 		FMovieSceneDoubleValue ExistKeyValue = GetValueFromChannel(KeyHandle);
 		FMovieSceneTangentData& ExistTangentData = ExistKeyValue.Tangent;
 		ExistTangentData.ArriveTangent = InNewValue;
-		SetValueInChannel(KeyHandle, ExistKeyValue);
+		SetValueInChannel(KeyHandle, OwningSection, ExistKeyValue);
 	}
 
 	/**
@@ -150,7 +150,7 @@ public:
 		FMovieSceneDoubleValue ExistKeyValue = GetValueFromChannel(KeyHandle);
 		FMovieSceneTangentData& ExistTangentData = ExistKeyValue.Tangent;
 		ExistTangentData.LeaveTangent = InNewValue;
-		SetValueInChannel(KeyHandle, ExistKeyValue);
+		SetValueInChannel(KeyHandle, OwningSection, ExistKeyValue);
 	}
 
 	/**
@@ -173,7 +173,7 @@ public:
 		FMovieSceneDoubleValue ExistKeyValue = GetValueFromChannel(KeyHandle);
 		FMovieSceneTangentData& ExistTangentData = ExistKeyValue.Tangent;
 		ExistTangentData.TangentWeightMode = InNewValue;
-		SetValueInChannel(KeyHandle, ExistKeyValue);
+		SetValueInChannel(KeyHandle, OwningSection, ExistKeyValue);
 	}
 
 	/**
@@ -196,7 +196,7 @@ public:
 		FMovieSceneDoubleValue ExistKeyValue = GetValueFromChannel(KeyHandle);
 		FMovieSceneTangentData& ExistTangentData = ExistKeyValue.Tangent;
 		ExistTangentData.ArriveTangentWeight = InNewValue;
-		SetValueInChannel(KeyHandle, ExistKeyValue);
+		SetValueInChannel(KeyHandle, OwningSection, ExistKeyValue);
 	}
 
 	/**
@@ -219,7 +219,7 @@ public:
 		FMovieSceneDoubleValue ExistKeyValue = GetValueFromChannel(KeyHandle);
 		FMovieSceneTangentData& ExistTangentData = ExistKeyValue.Tangent;
 		ExistTangentData.LeaveTangentWeight = InNewValue;
-		SetValueInChannel(KeyHandle, ExistKeyValue);
+		SetValueInChannel(KeyHandle, OwningSection, ExistKeyValue);
 	}
 };
 
@@ -390,7 +390,7 @@ public:
 	UFUNCTION(BlueprintCallable, Category = "Sequencer|Keys", meta = (DisplayName = "Remove Default (Double)"))
 	void RemoveDefault()
 	{
-		Impl::RemoveDefaultFromChannel(ChannelHandle);
+		Impl::RemoveDefaultFromChannel(ChannelHandle, OwningSequence, OwningSection);
 	}
 
 	/**
diff --git a/Engine/Plugins/MovieScene/SequencerScripting/Source/SequencerScripting/Private/KeysAndChannels/MovieSceneScriptingEvent.h b/Engine/Plugins/MovieScene/SequencerScripting/Source/SequencerScripting/Private/KeysAndChannels/MovieSceneScriptingEvent.h
--- a/Engine/Plugins/MovieScene/SequencerScripting/Source/SequencerScripting/Private/KeysAndChannels/MovieSceneScriptingEvent.h
+++ b/Engine/Plugins/MovieScene/SequencerScripting/Source/SequencerScripting/Private/KeysAndChannels/MovieSceneScriptingEvent.h
@@ -32,7 +32,7 @@ public:
 	* @param TimeUnit		Should the NewFrameNumber be interpreted as Display Rate frames or in Tick Resolution?
 	*/
 	UFUNCTION(BlueprintCallable, Category = "Sequencer|Keys", meta = (DisplayName = "Set Time (Event)"))
-	void SetTime(const FFrameNumber& NewFrameNumber, float SubFrame = 0.f, ESequenceTimeUnit TimeUnit = ESequenceTimeUnit::DisplayRate) { SetTimeInChannel(KeyHandle, OwningSequence, NewFrameNumber, TimeUnit, SubFrame); }
+	void SetTime(const FFrameNumber& NewFrameNumber, float SubFrame = 0.f, ESequenceTimeUnit TimeUnit = ESequenceTimeUnit::DisplayRate) { SetTimeInChannel(KeyHandle, OwningSequence, OwningSection, NewFrameNumber, TimeUnit, SubFrame); }
 
 	/**
 	* Gets the value for this key from the owning channel.
@@ -51,7 +51,7 @@ public:
 	UFUNCTION(BlueprintCallable, Category = "Sequencer|Keys", meta = (DisplayName = "Set Value (Event)"))
 	void SetValue(const FMovieSceneEvent& InNewValue)
 	{
-		SetValueInChannel(KeyHandle, InNewValue);
+		SetValueInChannel(KeyHandle, OwningSection, InNewValue);
 	}
 };
 
diff --git a/Engine/Plugins/MovieScene/SequencerScripting/Source/SequencerScripting/Private/KeysAndChannels/MovieSceneScriptingFloat.h b/Engine/Plugins/MovieScene/SequencerScripting/Source/SequencerScripting/Private/KeysAndChannels/MovieSceneScriptingFloat.h
--- a/Engine/Plugins/MovieScene/SequencerScripting/Source/SequencerScripting/Private/KeysAndChannels/MovieSceneScriptingFloat.h
+++ b/Engine/Plugins/MovieScene/SequencerScripting/Source/SequencerScripting/Private/KeysAndChannels/MovieSceneScriptingFloat.h
@@ -164,7 +164,7 @@ public:
 	}
 	virtual void SetTime(const FFrameNumber& NewFrameNumber, float SubFrame = 0.f, ESequenceTimeUnit TimeUnit = ESequenceTimeUnit::DisplayRate) override
 	{
-		SetTimeInChannel(KeyHandle, OwningSequence, NewFrameNumber, TimeUnit, SubFrame);
+		SetTimeInChannel(KeyHandle, OwningSequence, OwningSection, NewFrameNumber, TimeUnit, SubFrame);
 	}
 	virtual float GetValue() const override
 	{
@@ -174,7 +174,7 @@ public:
 	{
 		FMovieSceneFloatValue ExistValue = GetValueFromChannel(KeyHandle);
 		ExistValue.Value = InNewValue;
-		SetValueInChannel(KeyHandle, ExistValue);
+		SetValueInChannel(KeyHandle, OwningSection, ExistValue);
 	}
 	virtual ERichCurveInterpMode GetInterpolationMode() const override
 	{
@@ -184,7 +184,7 @@ public:
 	{
 		FMovieSceneFloatValue ExistValue = GetValueFromChannel(KeyHandle);
 		ExistValue.InterpMode = InNewValue;
-		SetValueInChannel(KeyHandle, ExistValue);
+		SetValueInChannel(KeyHandle, OwningSection, ExistValue);
 	}
 	virtual ERichCurveTangentMode GetTangentMode() const override
 	{
@@ -194,7 +194,7 @@ public:
 	{
 		FMovieSceneFloatValue ExistValue = GetValueFromChannel(KeyHandle);
 		ExistValue.TangentMode = InNewValue;
-		SetValueInChannel(KeyHandle, ExistValue);
+		SetValueInChannel(KeyHandle, OwningSection, ExistValue);
 	}
 	virtual float GetArriveTangent() const override
 	{
@@ -205,7 +205,7 @@ public:
 		FMovieSceneFloatValue ExistKeyValue = GetValueFromChannel(KeyHandle);
 		FMovieSceneTangentData& ExistTangentData = ExistKeyValue.Tangent;
 		ExistTangentData.ArriveTangent = InNewValue;
-		SetValueInChannel(KeyHandle, ExistKeyValue);
+		SetValueInChannel(KeyHandle, OwningSection, ExistKeyValue);
 	}
 	virtual float GetLeaveTangent() const override
 	{
@@ -216,7 +216,7 @@ public:
 		FMovieSceneFloatValue ExistKeyValue = GetValueFromChannel(KeyHandle);
 		FMovieSceneTangentData& ExistTangentData = ExistKeyValue.Tangent;
 		ExistTangentData.LeaveTangent = InNewValue;
-		SetValueInChannel(KeyHandle, ExistKeyValue);
+		SetValueInChannel(KeyHandle, OwningSection, ExistKeyValue);
 	}
 	virtual ERichCurveTangentWeightMode GetTangentWeightMode() const override
 	{
@@ -227,7 +227,7 @@ public:
 		FMovieSceneFloatValue ExistKeyValue = GetValueFromChannel(KeyHandle);
 		FMovieSceneTangentData& ExistTangentData = ExistKeyValue.Tangent;
 		ExistTangentData.TangentWeightMode = InNewValue;
-		SetValueInChannel(KeyHandle, ExistKeyValue);
+		SetValueInChannel(KeyHandle, OwningSection, ExistKeyValue);
 	}
 	virtual float GetArriveTangentWeight() const override
 	{
@@ -238,7 +238,7 @@ public:
 		FMovieSceneFloatValue ExistKeyValue = GetValueFromChannel(KeyHandle);
 		FMovieSceneTangentData& ExistTangentData = ExistKeyValue.Tangent;
 		ExistTangentData.ArriveTangentWeight = InNewValue;
-		SetValueInChannel(KeyHandle, ExistKeyValue);
+		SetValueInChannel(KeyHandle, OwningSection, ExistKeyValue);
 	}
 	virtual float GetLeaveTangentWeight() const override
 	{
@@ -249,7 +249,7 @@ public:
 		FMovieSceneFloatValue ExistKeyValue = GetValueFromChannel(KeyHandle);
 		FMovieSceneTangentData& ExistTangentData = ExistKeyValue.Tangent;
 		ExistTangentData.LeaveTangentWeight = InNewValue;
-		SetValueInChannel(KeyHandle, ExistKeyValue);
+		SetValueInChannel(KeyHandle, OwningSection, ExistKeyValue);
 	}
 };
 
@@ -264,7 +264,7 @@ public:
 	}
 	virtual void SetTime(const FFrameNumber& NewFrameNumber, float SubFrame = 0.f, ESequenceTimeUnit TimeUnit = ESequenceTimeUnit::DisplayRate) override
 	{
-		SetTimeInChannel(KeyHandle, OwningSequence, NewFrameNumber, TimeUnit, SubFrame);
+		SetTimeInChannel(KeyHandle, OwningSequence, OwningSection, NewFrameNumber, TimeUnit, SubFrame);
 	}
 	virtual float GetValue() const override
 	{
@@ -274,7 +274,7 @@ public:
 	{
 		FMovieSceneDoubleValue ExistValue = GetValueFromChannel(KeyHandle);
 		ExistValue.Value = (double)InNewValue;
-		SetValueInChannel(KeyHandle, ExistValue);
+		SetValueInChannel(KeyHandle, OwningSection, ExistValue);
 	}
 	virtual ERichCurveInterpMode GetInterpolationMode() const override
 	{
@@ -284,7 +284,7 @@ public:
 	{
 		FMovieSceneDoubleValue ExistValue = GetValueFromChannel(KeyHandle);
 		ExistValue.InterpMode = InNewValue;
-		SetValueInChannel(KeyHandle, ExistValue);
+		SetValueInChannel(KeyHandle, OwningSection, ExistValue);
 	}
 	virtual ERichCurveTangentMode GetTangentMode() const override
 	{
@@ -294,7 +294,7 @@ public:
 	{
 		FMovieSceneDoubleValue ExistValue = GetValueFromChannel(KeyHandle);
 		ExistValue.TangentMode = InNewValue;
-		SetValueInChannel(KeyHandle, ExistValue);
+		SetValueInChannel(KeyHandle, OwningSection, ExistValue);
 	}
 	virtual float GetArriveTangent() const override
 	{
@@ -305,7 +305,7 @@ public:
 		FMovieSceneDoubleValue ExistKeyValue = GetValueFromChannel(KeyHandle);
 		FMovieSceneTangentData& ExistTangentData = ExistKeyValue.Tangent;
 		ExistTangentData.ArriveTangent = InNewValue;
-		SetValueInChannel(KeyHandle, ExistKeyValue);
+		SetValueInChannel(KeyHandle, OwningSection, ExistKeyValue);
 	}
 	virtual float GetLeaveTangent() const override
 	{
@@ -316,7 +316,7 @@ public:
 		FMovieSceneDoubleValue ExistKeyValue = GetValueFromChannel(KeyHandle);
 		FMovieSceneTangentData& ExistTangentData = ExistKeyValue.Tangent;
 		ExistTangentData.LeaveTangent = InNewValue;
-		SetValueInChannel(KeyHandle, ExistKeyValue);
+		SetValueInChannel(KeyHandle, OwningSection, ExistKeyValue);
 	}
 	virtual ERichCurveTangentWeightMode GetTangentWeightMode() const override
 	{
@@ -327,7 +327,7 @@ public:
 		FMovieSceneDoubleValue ExistKeyValue = GetValueFromChannel(KeyHandle);
 		FMovieSceneTangentData& ExistTangentData = ExistKeyValue.Tangent;
 		ExistTangentData.TangentWeightMode = InNewValue;
-		SetValueInChannel(KeyHandle, ExistKeyValue);
+		SetValueInChannel(KeyHandle, OwningSection, ExistKeyValue);
 	}
 	virtual float GetArriveTangentWeight() const override
 	{
@@ -338,7 +338,7 @@ public:
 		FMovieSceneDoubleValue ExistKeyValue = GetValueFromChannel(KeyHandle);
 		FMovieSceneTangentData& ExistTangentData = ExistKeyValue.Tangent;
 		ExistTangentData.ArriveTangentWeight = InNewValue;
-		SetValueInChannel(KeyHandle, ExistKeyValue);
+		SetValueInChannel(KeyHandle, OwningSection, ExistKeyValue);
 	}
 	virtual float GetLeaveTangentWeight() const override
 	{
@@ -349,7 +349,7 @@ public:
 		FMovieSceneDoubleValue ExistKeyValue = GetValueFromChannel(KeyHandle);
 		FMovieSceneTangentData& ExistTangentData = ExistKeyValue.Tangent;
 		ExistTangentData.LeaveTangentWeight = InNewValue;
-		SetValueInChannel(KeyHandle, ExistKeyValue);
+		SetValueInChannel(KeyHandle, OwningSection, ExistKeyValue);
 	}
 };
 
@@ -616,11 +616,11 @@ public:
 	{
 		if (FloatChannelHandle.Get())
 		{
-			FloatImpl::RemoveDefaultFromChannel(FloatChannelHandle);
+			FloatImpl::RemoveDefaultFromChannel(FloatChannelHandle, OwningSequence, OwningSection);
 		}
 		else
 		{
-			DoubleImpl::RemoveDefaultFromChannel(DoubleChannelHandle);
+			DoubleImpl::RemoveDefaultFromChannel(DoubleChannelHandle, OwningSequence, OwningSection);
 		}
 	}
 
diff --git a/Engine/Plugins/MovieScene/SequencerScripting/Source/SequencerScripting/Private/KeysAndChannels/MovieSceneScriptingInteger.h b/Engine/Plugins/MovieScene/SequencerScripting/Source/SequencerScripting/Private/KeysAndChannels/MovieSceneScriptingInteger.h
--- a/Engine/Plugins/MovieScene/SequencerScripting/Source/SequencerScripting/Private/KeysAndChannels/MovieSceneScriptingInteger.h
+++ b/Engine/Plugins/MovieScene/SequencerScripting/Source/SequencerScripting/Private/KeysAndChannels/MovieSceneScriptingInteger.h
@@ -37,7 +37,7 @@ public:
 	* @param TimeUnit		Should the NewFrameNumber be interpreted as Display Rate frames or in Tick Resolution?
 	*/
 	UFUNCTION(BlueprintCallable, Category = "Sequencer|Keys", meta = (DisplayName = "Set Time (Integer)"))
-	void SetTime(const FFrameNumber& NewFrameNumber, float SubFrame = 0.f, ESequenceTimeUnit TimeUnit = ESequenceTimeUnit::DisplayRate) { SetTimeInChannel(KeyHandle, OwningSequence, NewFrameNumber, TimeUnit, SubFrame); }
+	void SetTime(const FFrameNumber& NewFrameNumber, float SubFrame = 0.f, ESequenceTimeUnit TimeUnit = ESequenceTimeUnit::DisplayRate) { SetTimeInChannel(KeyHandle, OwningSequence, OwningSection, NewFrameNumber, TimeUnit, SubFrame); }
 
 	/**
 	* Gets the value for this key from the owning channel.
@@ -56,7 +56,7 @@ public:
 	UFUNCTION(BlueprintCallable, Category = "Sequencer|Keys", meta = (DisplayName = "Set Value (Integer)"))
 	void SetValue(int32 InNewValue)
 	{
-		SetValueInChannel(KeyHandle, InNewValue);
+		SetValueInChannel(KeyHandle, OwningSection, InNewValue);
 	}
 };
 
@@ -129,7 +129,7 @@ public:
 	UFUNCTION(BlueprintCallable, Category = "Sequencer|Keys", meta = (DisplayName = "Remove Default (Integer)"))
 	void RemoveDefault()
 	{
-		Impl::RemoveDefaultFromChannel(ChannelHandle);
+		Impl::RemoveDefaultFromChannel(ChannelHandle, OwningSequence, OwningSection);
 	}
 
 	/**
diff --git a/Engine/Plugins/MovieScene/SequencerScripting/Source/SequencerScripting/Private/KeysAndChannels/MovieSceneScriptingObjectPath.h b/Engine/Plugins/MovieScene/SequencerScripting/Source/SequencerScripting/Private/KeysAndChannels/MovieSceneScriptingObjectPath.h
--- a/Engine/Plugins/MovieScene/SequencerScripting/Source/SequencerScripting/Private/KeysAndChannels/MovieSceneScriptingObjectPath.h
+++ b/Engine/Plugins/MovieScene/SequencerScripting/Source/SequencerScripting/Private/KeysAndChannels/MovieSceneScriptingObjectPath.h
@@ -37,7 +37,7 @@ public:
 	 * @param TimeUnit		Should the NewFrameNumber be interpreted as Display Rate frames or in Tick Resolution?
 	 */
 	UFUNCTION(BlueprintCallable, Category = "Sequencer|Keys", meta = (DisplayName = "Set Time (Object Path)"))
-	void SetTime(const FFrameNumber& NewFrameNumber, float SubFrame = 0.f, ESequenceTimeUnit TimeUnit = ESequenceTimeUnit::DisplayRate) { SetTimeInChannel(KeyHandle, OwningSequence, NewFrameNumber, TimeUnit, SubFrame); }
+	void SetTime(const FFrameNumber& NewFrameNumber, float SubFrame = 0.f, ESequenceTimeUnit TimeUnit = ESequenceTimeUnit::DisplayRate) { SetTimeInChannel(KeyHandle, OwningSequence, OwningSection, NewFrameNumber, TimeUnit, SubFrame); }
 
 	/**
 	 * Gets the value for this key from the owning channel.
@@ -58,7 +58,7 @@ public:
 	void SetValue(UObject* InNewValue)
 	{
 		FMovieSceneObjectPathChannelKeyValue ReferenceKey = FMovieSceneObjectPathChannelKeyValue(InNewValue);
-		SetValueInChannel(KeyHandle, ReferenceKey);
+		SetValueInChannel(KeyHandle, OwningSection, ReferenceKey);
 	}
 };
 
@@ -156,12 +156,25 @@ public:
 	UFUNCTION(BlueprintCallable, Category = "Sequencer|Keys", meta = (DisplayName = "Remove Default (Object Path)"))
 	void RemoveDefault()
 	{
+		if (OwningSection.IsValid())
+		{
+			OwningSection->Modify();
+		}
+
 		// FMovieSceneObjectPathChannelKeyValue doesn't implement RemoveDefault, instead it implements ClearDefault(). Wrapping this function by hand,
 		// and not falling back to the template as a result, but keeping the same function name so it is consistent with the other scripting channels.
 		FMovieSceneObjectPathChannel* Channel = ChannelHandle.Get();
 		if (Channel)
 		{
 			Channel->ClearDefault();
+
+#if WITH_EDITOR
+			const FMovieSceneChannelMetaData* MetaData = ChannelHandle.GetMetaData();
+			if (MetaData && OwningSection.IsValid() && OwningSequence.IsValid() && OwningSequence->GetMovieScene())
+			{
+				OwningSequence->GetMovieScene()->OnChannelChanged().Broadcast(MetaData, OwningSection.Get());
+			}
+#endif
 			return;
 		}
 		FFrame::KismetExecutionMessage(TEXT("Invalid ChannelHandle for MovieSceneScriptingChannel, failed to remove default value."), ELogVerbosity::Error);
diff --git a/Engine/Plugins/MovieScene/SequencerScripting/Source/SequencerScripting/Private/KeysAndChannels/MovieSceneScriptingString.h b/Engine/Plugins/MovieScene/SequencerScripting/Source/SequencerScripting/Private/KeysAndChannels/MovieSceneScriptingString.h
--- a/Engine/Plugins/MovieScene/SequencerScripting/Source/SequencerScripting/Private/KeysAndChannels/MovieSceneScriptingString.h
+++ b/Engine/Plugins/MovieScene/SequencerScripting/Source/SequencerScripting/Private/KeysAndChannels/MovieSceneScriptingString.h
@@ -36,7 +36,7 @@ public:
 	* @param TimeUnit		Should the NewFrameNumber be interpreted as Display Rate frames or in Tick Resolution?
 	*/
 	UFUNCTION(BlueprintCallable, Category = "Sequencer|Keys", meta = (DisplayName = "Set Time (String)"))
-	void SetTime(const FFrameNumber& NewFrameNumber, float SubFrame = 0.f, ESequenceTimeUnit TimeUnit = ESequenceTimeUnit::DisplayRate) { SetTimeInChannel(KeyHandle, OwningSequence, NewFrameNumber, TimeUnit, SubFrame); }
+	void SetTime(const FFrameNumber& NewFrameNumber, float SubFrame = 0.f, ESequenceTimeUnit TimeUnit = ESequenceTimeUnit::DisplayRate) { SetTimeInChannel(KeyHandle, OwningSequence, OwningSection, NewFrameNumber, TimeUnit, SubFrame); }
 
 	/**
 	* Gets the value for this key from the owning channel.
@@ -55,7 +55,7 @@ public:
 	UFUNCTION(BlueprintCallable, Category = "Sequencer|Keys", meta = (DisplayName = "Set Value (String)"))
 	void SetValue(const FString& InNewValue)
 	{
-		SetValueInChannel(KeyHandle, InNewValue);
+		SetValueInChannel(KeyHandle, OwningSection, InNewValue);
 	}
 };
 
@@ -129,7 +129,7 @@ public:
 	UFUNCTION(BlueprintCallable, Category = "Sequencer|Keys", meta = (DisplayName = "Remove Default (String)"))
 	void RemoveDefault()
 	{
-		Impl::RemoveDefaultFromChannel(ChannelHandle);
+		Impl::RemoveDefaultFromChannel(ChannelHandle, OwningSequence, OwningSection);
 	}
 
 	/**
