Sequencer: Remove deprecated at 5.0 and below from movie scene and level sequence

#rnx

diff --git a/Engine/Plugins/MovieScene/SequencerScripting/Source/SequencerScripting/Private/ExtensionLibraries/MovieSceneSectionExtensions.cpp b/Engine/Plugins/MovieScene/SequencerScripting/Source/SequencerScripting/Private/ExtensionLibraries/MovieSceneSectionExtensions.cpp
--- a/Engine/Plugins/MovieScene/SequencerScripting/Source/SequencerScripting/Private/ExtensionLibraries/MovieSceneSectionExtensions.cpp
+++ b/Engine/Plugins/MovieScene/SequencerScripting/Source/SequencerScripting/Private/ExtensionLibraries/MovieSceneSectionExtensions.cpp
@@ -20,18 +20,6 @@
 #include UE_INLINE_GENERATED_CPP_BY_NAME(MovieSceneSectionExtensions)
 
 
-FSequencerScriptingRange UMovieSceneSectionExtensions::GetRange(UMovieSceneSection* Section)
-{
-	if (!Section)
-	{
-		FFrame::KismetExecutionMessage(TEXT("Cannot call GetRange on a null section"), ELogVerbosity::Error);
-		return FSequencerScriptingRange();
-	}
-
-	UMovieScene* MovieScene = Section->GetTypedOuter<UMovieScene>();
-	return FSequencerScriptingRange::FromNative(Section->GetRange(), MovieScene->GetTickResolution());
-}
-
 bool UMovieSceneSectionExtensions::HasStartFrame(UMovieSceneSection* Section)
 {
 	if (!Section)
@@ -498,35 +486,6 @@ void GetScriptingChannelsForChannel(FMovieSceneChannelProxy& ChannelProxy,TWeakO
 	}
 }
 
-TArray<UMovieSceneScriptingChannel*> UMovieSceneSectionExtensions::GetChannels(UMovieSceneSection* Section)
-{
-	TArray<UMovieSceneScriptingChannel*> Channels;
-	if (!Section)
-	{
-		FFrame::KismetExecutionMessage(TEXT("Cannot get channels for null section"), ELogVerbosity::Error);
-		return Channels;
-	}
-	FMovieSceneChannelProxy& ChannelProxy = Section->GetChannelProxy();
-
-	// ToDo: This isn't a great way of collecting all of the channel types and instantiating the correct UObject for it
-	// but given that we're already hard-coding the UObject implementations... We currently support the following channel
-	// types: Bool, Byte, Float, Double, Integer, String, Event, ActorReference
-	TWeakObjectPtr<UMovieSceneSequence> Sequence = Section->GetTypedOuter<UMovieSceneSequence>();
-	GetScriptingChannelsForChannel<FMovieSceneBoolChannel, UMovieSceneScriptingBoolChannel>(ChannelProxy, Sequence, Section, Channels);
-	GetScriptingChannelsForChannel<FMovieSceneByteChannel, UMovieSceneScriptingByteChannel>(ChannelProxy, Sequence, Section, Channels);
-	GetScriptingChannelsForChannel<FMovieSceneFloatChannel, UMovieSceneScriptingFloatChannel>(ChannelProxy, Sequence, Section, Channels);
-	GetScriptingChannelsForChannel<FMovieSceneIntegerChannel, UMovieSceneScriptingIntegerChannel>(ChannelProxy, Sequence, Section, Channels);
-	GetScriptingChannelsForChannel<FMovieSceneStringChannel, UMovieSceneScriptingStringChannel>(ChannelProxy, Sequence, Section, Channels);
-	GetScriptingChannelsForChannel<FMovieSceneEventChannel, UMovieSceneScriptingEventChannel>(ChannelProxy, Sequence, Section, Channels);
-	GetScriptingChannelsForChannel<FMovieSceneActorReferenceData, UMovieSceneScriptingActorReferenceChannel>(ChannelProxy, Sequence, Section, Channels);
-	GetScriptingChannelsForChannel<FMovieSceneObjectPathChannel, UMovieSceneScriptingObjectPathChannel>(ChannelProxy, Sequence, Section, Channels);
-
-	// Wrap double channels around float scripting channels.
-	GetScriptingChannelsForChannel<FMovieSceneDoubleChannel, UMovieSceneScriptingFloatChannel>(ChannelProxy, Sequence, Section, Channels);
-
-	return Channels;
-}
-
 TArray<UMovieSceneScriptingChannel*> UMovieSceneSectionExtensions::GetAllChannels(UMovieSceneSection* Section)
 {
 	TArray<UMovieSceneScriptingChannel*> Channels;
@@ -554,38 +513,6 @@ TArray<UMovieSceneScriptingChannel*> UMovieSceneSectionExtensions::GetAllChannel
 	return Channels;
 }
 
-TArray<UMovieSceneScriptingChannel*> UMovieSceneSectionExtensions::FindChannelsByType(UMovieSceneSection* Section, TSubclassOf<UMovieSceneScriptingChannel> ChannelType)
-{
-	TArray<UMovieSceneScriptingChannel*> Channels;
-	if (!Section)
-	{
-		FFrame::KismetExecutionMessage(TEXT("Cannot get channels for null section"), ELogVerbosity::Error);
-		return Channels;
-	}
- 
-	FMovieSceneChannelProxy& ChannelProxy = Section->GetChannelProxy();
- 
-	// ToDo: This needs to be done dynamically in the future but there's not a good way to do that right now with all of the SFINAE-based Templating driving these channels
-	TWeakObjectPtr<UMovieSceneSequence> Sequence = Section->GetTypedOuter<UMovieSceneSequence>();
-
-	if (ChannelType == UMovieSceneScriptingBoolChannel::StaticClass())					{ GetScriptingChannelsForChannel<FMovieSceneBoolChannel, UMovieSceneScriptingBoolChannel>(ChannelProxy, Sequence, Section, Channels); }
-	else if (ChannelType == UMovieSceneScriptingByteChannel::StaticClass())				{ GetScriptingChannelsForChannel<FMovieSceneByteChannel, UMovieSceneScriptingByteChannel>(ChannelProxy, Sequence, Section, Channels); }
-	else if (ChannelType == UMovieSceneScriptingFloatChannel::StaticClass())			{ GetScriptingChannelsForChannel<FMovieSceneFloatChannel, UMovieSceneScriptingFloatChannel>(ChannelProxy, Sequence, Section, Channels); }
-	else if (ChannelType == UMovieSceneScriptingDoubleChannel::StaticClass())			{ GetScriptingChannelsForChannel<FMovieSceneDoubleChannel, UMovieSceneScriptingFloatChannel>(ChannelProxy, Sequence, Section, Channels); }
-	else if (ChannelType == UMovieSceneScriptingIntegerChannel::StaticClass())			{ GetScriptingChannelsForChannel<FMovieSceneIntegerChannel, UMovieSceneScriptingIntegerChannel>(ChannelProxy, Sequence, Section, Channels); }
-	else if (ChannelType == UMovieSceneScriptingStringChannel::StaticClass())			{ GetScriptingChannelsForChannel<FMovieSceneStringChannel, UMovieSceneScriptingStringChannel>(ChannelProxy, Sequence, Section, Channels); }
-	else if (ChannelType == UMovieSceneScriptingEventChannel::StaticClass())			{ GetScriptingChannelsForChannel<FMovieSceneEventChannel, UMovieSceneScriptingEventChannel>(ChannelProxy, Sequence, Section, Channels); }
-	else if (ChannelType == UMovieSceneScriptingActorReferenceChannel::StaticClass()) { GetScriptingChannelsForChannel<FMovieSceneActorReferenceData, UMovieSceneScriptingActorReferenceChannel>(ChannelProxy, Sequence, Section, Channels); }
-	else if (ChannelType == UMovieSceneScriptingObjectPathChannel::StaticClass()) { GetScriptingChannelsForChannel<FMovieSceneObjectPathChannel, UMovieSceneScriptingObjectPathChannel>(ChannelProxy, Sequence, Section,  Channels); }
-	else
-	{
-		FFrame::KismetExecutionMessage(TEXT("Unsupported ChannelType for FindChannelsByType!"), ELogVerbosity::Error);
-	}
-
-
-	return Channels;
-}
-
 TArray<UMovieSceneScriptingChannel*> UMovieSceneSectionExtensions::GetChannelsByType(UMovieSceneSection* Section, TSubclassOf<UMovieSceneScriptingChannel> ChannelType)
 {
 	TArray<UMovieSceneScriptingChannel*> Channels;
diff --git a/Engine/Plugins/MovieScene/SequencerScripting/Source/SequencerScripting/Private/ExtensionLibraries/MovieSceneSequenceExtensions.cpp b/Engine/Plugins/MovieScene/SequencerScripting/Source/SequencerScripting/Private/ExtensionLibraries/MovieSceneSequenceExtensions.cpp
--- a/Engine/Plugins/MovieScene/SequencerScripting/Source/SequencerScripting/Private/ExtensionLibraries/MovieSceneSequenceExtensions.cpp
+++ b/Engine/Plugins/MovieScene/SequencerScripting/Source/SequencerScripting/Private/ExtensionLibraries/MovieSceneSequenceExtensions.cpp
@@ -648,23 +648,6 @@ EUpdateClockSource UMovieSceneSequenceExtensions::GetClockSource(UMovieSceneSequ
 	return EUpdateClockSource::Tick;
 }
 
-FTimecode UMovieSceneSequenceExtensions::GetTimecodeSource(UMovieSceneSequence* Sequence)
-{
-	if (!Sequence)
-	{
-		FFrame::KismetExecutionMessage(TEXT("Cannot call GetTimecodeSource on a null sequence"), ELogVerbosity::Error);
-		return FTimecode();
-	}
-
-	UMovieScene* MovieScene = GetMovieScene(Sequence);
-	if (!MovieScene)
-	{
-		return FTimecode();
-	}
-
-	return MovieScene->GetEarliestTimecodeSource().Timecode;
-}
-
 FMovieSceneBindingProxy UMovieSceneSequenceExtensions::FindBindingByName(UMovieSceneSequence* Sequence, FString Name)
 {
 	if (!Sequence)
@@ -845,24 +828,6 @@ TArray<UObject*> UMovieSceneSequenceExtensions::LocateBoundObjects(UMovieSceneSe
 	return Result;
 }
 
-FMovieSceneObjectBindingID UMovieSceneSequenceExtensions::MakeBindingID(UMovieSceneSequence* RootSequence, const FMovieSceneBindingProxy& InBinding, EMovieSceneObjectBindingSpace Space)
-{
-	if (!RootSequence)
-	{
-		FFrame::KismetExecutionMessage(TEXT("Cannot call MakeBindingID on a null sequence"), ELogVerbosity::Error);
-		return FMovieSceneObjectBindingID();
-	}
-
-	// This function was kinda flawed before - when ::Local was passed for the Space parameter,
-	// and the sub sequence ID could not be found it would always fall back to a binding for ::Root without any Sequence ID
-	FMovieSceneObjectBindingID BindingID = GetPortableBindingID(RootSequence, RootSequence, InBinding);
-	if (Space == EMovieSceneObjectBindingSpace::Root)
-	{
-		BindingID.ReinterpretAsFixed();
-	}
-	return BindingID;
-}
-
 FMovieSceneObjectBindingID UMovieSceneSequenceExtensions::GetBindingID(const FMovieSceneBindingProxy& InBinding)
 {
 	return UE::MovieScene::FRelativeObjectBindingID(InBinding.BindingID);
diff --git a/Engine/Plugins/MovieScene/SequencerScripting/Source/SequencerScripting/Public/ExtensionLibraries/MovieSceneSectionExtensions.h b/Engine/Plugins/MovieScene/SequencerScripting/Source/SequencerScripting/Public/ExtensionLibraries/MovieSceneSectionExtensions.h
--- a/Engine/Plugins/MovieScene/SequencerScripting/Source/SequencerScripting/Public/ExtensionLibraries/MovieSceneSectionExtensions.h
+++ b/Engine/Plugins/MovieScene/SequencerScripting/Source/SequencerScripting/Public/ExtensionLibraries/MovieSceneSectionExtensions.h
@@ -21,9 +21,6 @@ class UMovieSceneSectionExtensions : public UBlueprintFunctionLibrary
 	GENERATED_BODY()
 public:
 
-	UE_DEPRECATED(4.22, "Please use GetStartFrame and GetEndFrame instead.")
-	static FSequencerScriptingRange GetRange(UMovieSceneSection* Section);
-
 	/**
 	 * Has start frame
 	 *
@@ -206,21 +203,6 @@ public:
 	UFUNCTION(BlueprintCallable, Category = "Sequencer|Section", meta = (ScriptMethod))
 	static void SetEndFrameBounded(UMovieSceneSection* Section, bool bIsBounded);
 
-	/**
-	* Find all channels that belong to the specified UMovieSceneSection. Some sections have many channels (such as
-	* Transforms containing 9 float channels to represent Translation/Rotation/Scale), and a section may have mixed
-	* channel types.
-	*
-	* For backwards compatibility reasons, all double channels will be wrapped inside float scripting channels.
-	* Use GetAllChannels to get proper double channels.
-	*
-	* @param Section       The section to use.
-	* @return An array containing any key channels that match the type specified
-	*/
-	UE_DEPRECATED(5.0, "Please use GetAllChannels instead and query double-precision channels as such where appropriate")
-	UFUNCTION(BlueprintCallable, Category = "Sequencer|Section", meta=(ScriptMethod, DevelopmentOnly, DeprecatedFunction, DeprecationMessage = "Please use GetAllChannels instead and query double-precision channels as such where appropriate"))
-	static TArray<UMovieSceneScriptingChannel*> GetChannels(UMovieSceneSection* Section);
-
 	/**
 	* Find all channels that belong to the specified UMovieSceneSection. Some sections have many channels (such as
 	* Transforms containing 9 double channels to represent Translation/Rotation/Scale), and a section may have mixed
@@ -232,21 +214,6 @@ public:
 	UFUNCTION(BlueprintCallable, Category = "Sequencer|Section", meta=(ScriptMethod, DevelopmentOnly))
 	static TArray<UMovieSceneScriptingChannel*> GetAllChannels(UMovieSceneSection* Section);
 
-	/**
-	* Find all channels that belong to the specified UMovieSceneSection that match the specific type. This will filter out any children who do not inherit
-	* from the specified type for you.
-	*
-	* For backwards compatibility reasons, all double channels will be wrapped inside float scripting channels.
-	* Use GetAllChannels to get proper double channels.
-	*
-	* @param Section        The section to use.
-	* @param ChannelType	The class type to look for.
-	* @return An array containing any key channels that match the type specified
-	*/
-	UE_DEPRECATED(5.0, "Please use GetChannelsByType instead")
-	UFUNCTION(BlueprintCallable, Category = "Sequencer|Section", meta = (ScriptMethod, DeterminesOutputType="TrackType", DevelopmentOnly, DeprecatedFunction, DeprecationMessage = "Please use GetChannelsByType instead"))
-	static TArray<UMovieSceneScriptingChannel*> FindChannelsByType(UMovieSceneSection* Section, TSubclassOf<UMovieSceneScriptingChannel> ChannelType);
-
 	/**
 	* Find all channels that belong to the specified UMovieSceneSection that match the specific type. This will filter out any children who do not inherit
 	* from the specified type for you.
diff --git a/Engine/Plugins/MovieScene/SequencerScripting/Source/SequencerScripting/Public/ExtensionLibraries/MovieSceneSequenceExtensions.h b/Engine/Plugins/MovieScene/SequencerScripting/Source/SequencerScripting/Public/ExtensionLibraries/MovieSceneSequenceExtensions.h
--- a/Engine/Plugins/MovieScene/SequencerScripting/Source/SequencerScripting/Public/ExtensionLibraries/MovieSceneSequenceExtensions.h
+++ b/Engine/Plugins/MovieScene/SequencerScripting/Source/SequencerScripting/Public/ExtensionLibraries/MovieSceneSequenceExtensions.h
@@ -361,16 +361,6 @@ public:
 	UFUNCTION(BlueprintPure, Category = "Sequencer|Sequence", meta = (ScriptMethod))
 	static EUpdateClockSource GetClockSource(UMovieSceneSequence* InSequence);
 
-	/**
-	 * Get the timecode source of this sequence
-	 *
-	 * @param Sequence        The sequence within which to get the timecode source
-	 * @return Timecode source of this sequence
-	 */
-	UE_DEPRECATED(5.0, "GetTimecodeSource() is no longer supported for movie scene sequences. Please use GetEarliestTimecodeSource() instead.")
-	UFUNCTION(BlueprintCallable, Category = "Sequencer|Sequence", meta = (ScriptMethod, DeprecatedFunction, DeprecationMessage = "GetTimecodeSource() is no longer supported for movie scene sequences. Please use GetEarliestTimecodeSource() instead."))
-	static FTimecode GetTimecodeSource(UMovieSceneSequence* Sequence);
-
 	/**
 	 * Attempt to locate a binding in this sequence by its name
 	 *
@@ -460,19 +450,6 @@ public:
 	static TArray<UObject*> LocateBoundObjects(UMovieSceneSequence* Sequence, const FMovieSceneBindingProxy& InBinding, UObject* Context);
 
 
-	/**
-	 * Make a binding id for the given binding in this sequence
-	 *
-	 * @param Sequence  The sequence that contains the sequence
-	 * @param Binding The binding proxy to generate the binding id from
-	 * @param Space The object binding space to resolve from (Root or Local)
-	 * @return The new object binding id
-	 */
-	UE_DEPRECATED(5.0, "Please migrate to GetBindingID or GetPortableBindingID depending on use-case.")
-	UFUNCTION(BlueprintCallable, Category = "Sequencer|Sequence", meta = (ScriptMethod, DeprecatedFunction, DeprecationMessage="Please migrate to GetBindingID or GetPortableBindingID depending on use-case."))
-	static FMovieSceneObjectBindingID MakeBindingID(UMovieSceneSequence* Sequence, const FMovieSceneBindingProxy& InBinding, EMovieSceneObjectBindingSpace Space = EMovieSceneObjectBindingSpace::Root);
-
-
 	/**
 	 * Get the binding ID for a binding within a sequence.
 	 * @note: The resulting binding is only valid when applied to properties within the same sequence as this binding. Use GetPortableBindingID for bindings which live in different sub-sequences.
diff --git a/Engine/Source/Editor/Sequencer/Public/ISequencerTrackEditor.h b/Engine/Source/Editor/Sequencer/Public/ISequencerTrackEditor.h
--- a/Engine/Source/Editor/Sequencer/Public/ISequencerTrackEditor.h
+++ b/Engine/Source/Editor/Sequencer/Public/ISequencerTrackEditor.h
@@ -174,9 +174,6 @@ public:
 	 */
 	virtual bool OnAllowDrop(const FDragDropEvent& DragDropEvent, FSequencerDragDropParams& DragDropParams) = 0;
 
-	UE_DEPRECATED(4.27, "Use OnAllowDrop with DragDropParams.")
-	virtual bool OnAllowDrop(const FDragDropEvent& DragDropEvent, UMovieSceneTrack* Track, int32 RowIndex, const FGuid& TargetObjectGuid) { FSequencerDragDropParams DragDropParams; return OnAllowDrop(DragDropEvent, DragDropParams); }
-
 	/**
 	 * Called when an asset is dropped directly onto a track.
 	 *
@@ -186,9 +183,6 @@ public:
 	 */	
 	virtual FReply OnDrop(const FDragDropEvent& DragDropEvent, const FSequencerDragDropParams& DragDropParams) = 0;
 
-	UE_DEPRECATED(4.27, "Use OnDrop with DragDropParams.")
-	virtual FReply OnDrop(const FDragDropEvent& DragDropEvent, UMovieSceneTrack* Track, int32 RowIndex, const FGuid& TargetObjectGuid) { return OnDrop(DragDropEvent, FSequencerDragDropParams()); }
-
 	/**
 	 * Called to generate a section layout for a particular section.
 	 *
@@ -207,9 +201,6 @@ public:
 	/** Called when the instance of this track editor is released */
 	virtual void OnRelease() = 0;
 
-	UE_DEPRECATED(5.0, "This function is no longer used.")
-	virtual int32 PaintTrackArea(const FPaintArgs& Args, const FGeometry& AllottedGeometry, const FSlateRect& MyCullingRect, FSlateWindowElementList& OutDrawElements, int32 LayerId, const FWidgetStyle& InWidgetStyle) { return LayerId; }
-
 	/**
 	 * Returns whether a track class is supported by this tool.
 	 *
diff --git a/Engine/Source/Runtime/AnimGraphRuntime/Private/AnimSequencerInstance.cpp b/Engine/Source/Runtime/AnimGraphRuntime/Private/AnimSequencerInstance.cpp
--- a/Engine/Source/Runtime/AnimGraphRuntime/Private/AnimSequencerInstance.cpp
+++ b/Engine/Source/Runtime/AnimGraphRuntime/Private/AnimSequencerInstance.cpp
@@ -37,12 +37,6 @@ void UAnimSequencerInstance::UpdateAnimTrack(UAnimSequenceBase* InAnimSequence,
 	GetProxyOnGameThread<FAnimSequencerInstanceProxy>().UpdateAnimTrack(InAnimSequence, SequenceId, InFromPosition, InToPosition, Weight, bFireNotifies);
 }
 
-void UAnimSequencerInstance::UpdateAnimTrackWithRootMotion(UAnimSequenceBase* InAnimSequence, int32 SequenceId, const TOptional<FRootMotionOverride>& RootMotion, float InFromPosition, float InToPosition, float Weight, bool bFireNotifies)
-{
-	FAnimSequencerData AnimSequencerData(InAnimSequence, SequenceId, RootMotion, InFromPosition, InToPosition, Weight, bFireNotifies, ESwapRootBone::SwapRootBone_None, FTransform(), nullptr);
-	GetProxyOnGameThread<FAnimSequencerInstanceProxy>().UpdateAnimTrackWithRootMotion(AnimSequencerData);
-}
-
 void UAnimSequencerInstance::UpdateAnimTrackWithRootMotion(const FAnimSequencerData& InAnimSequencerData)
 {
 	GetProxyOnGameThread<FAnimSequencerInstanceProxy>().UpdateAnimTrackWithRootMotion(InAnimSequencerData);
diff --git a/Engine/Source/Runtime/AnimGraphRuntime/Public/AnimSequencerInstance.h b/Engine/Source/Runtime/AnimGraphRuntime/Public/AnimSequencerInstance.h
--- a/Engine/Source/Runtime/AnimGraphRuntime/Public/AnimSequencerInstance.h
+++ b/Engine/Source/Runtime/AnimGraphRuntime/Public/AnimSequencerInstance.h
@@ -26,9 +26,6 @@ public:
 	virtual void UpdateAnimTrack(UAnimSequenceBase* InAnimSequence, int32 SequenceId, float InFromPosition, float InToPosition, float Weight, bool bFireNotifies);
 
 	/** Update with Root Motion*/
-	UE_DEPRECATED(5.0, "Please use UpdateAnimTrackWithRootMotion that takes a MirrorDataTable")
-	void UpdateAnimTrackWithRootMotion(UAnimSequenceBase* InAnimSequence, int32 SequenceId, const TOptional<FRootMotionOverride>& RootMotion, float InFromPosition, float InToPosition, float Weight, bool bFireNotifies);
-
 	UE_DEPRECATED(5.1, "Please use the UpdateAnimTrackWithRootMotion that takes FAnimSequencerData")
 	void UpdateAnimTrackWithRootMotion(UAnimSequenceBase* InAnimSequence, int32 SequenceId,const TOptional<FRootMotionOverride>& RootMotion, float InFromPosition, float InToPosition, float Weight, bool bFireNotifies, UMirrorDataTable* InMirrorDataTable);
 	
diff --git a/Engine/Source/Runtime/LevelSequence/Public/LevelSequenceActor.h b/Engine/Source/Runtime/LevelSequence/Public/LevelSequenceActor.h
--- a/Engine/Source/Runtime/LevelSequence/Public/LevelSequenceActor.h
+++ b/Engine/Source/Runtime/LevelSequence/Public/LevelSequenceActor.h
@@ -132,16 +132,6 @@ public:
 	UFUNCTION(BlueprintCallable, Category="Sequencer|Player")
 	ULevelSequence* GetSequence() const;
 
-	/**
-	 * Get the level sequence being played by this actor.
-	 *
-	 * @return Level sequence, or nullptr if not assigned or if it cannot be loaded.
-	 * @see SetSequence
-	 */
-	UE_DEPRECATED(5.0, "LoadSequence has been deprecated, please use GetSequence")
-	UFUNCTION(BlueprintCallable, Category="Sequencer|Player", meta=(DeprecatedFunction))
-	ULevelSequence* LoadSequence() const { return GetSequence(); }
-
 	/**
 	 * Set the level sequence being played by this actor.
 	 *
diff --git a/Engine/Source/Runtime/LevelSequence/Public/LevelSequencePlayer.h b/Engine/Source/Runtime/LevelSequence/Public/LevelSequencePlayer.h
--- a/Engine/Source/Runtime/LevelSequence/Public/LevelSequencePlayer.h
+++ b/Engine/Source/Runtime/LevelSequence/Public/LevelSequencePlayer.h
@@ -27,10 +27,6 @@ class UCameraComponent;
 
 DECLARE_DYNAMIC_MULTICAST_DELEGATE_OneParam(FOnLevelSequencePlayerCameraCutEvent, UCameraComponent*, CameraComponent);
 
-USTRUCT(BlueprintType)
-struct UE_DEPRECATED(5.0, "Snapshot settings are deprecated. Use the frame rate from the FQualifiedFrameTime on the RootTime and the ShotTime") FLevelSequenceSnapshotSettings
-{ GENERATED_BODY() };
-
 /**
  * Frame snapshot information for a level sequence
  */
diff --git a/Engine/Source/Runtime/MovieScene/Private/MovieScene.cpp b/Engine/Source/Runtime/MovieScene/Private/MovieScene.cpp
--- a/Engine/Source/Runtime/MovieScene/Private/MovieScene.cpp
+++ b/Engine/Source/Runtime/MovieScene/Private/MovieScene.cpp
@@ -1425,13 +1425,6 @@ void UMovieScene::RemoveNullTracks()
  *****************************************************************************/
 
 
-void UMovieScene::PreSave(const class ITargetPlatform* TargetPlatform)
-{
-	PRAGMA_DISABLE_DEPRECATION_WARNINGS;
-	Super::PreSave(TargetPlatform);
-	PRAGMA_ENABLE_DEPRECATION_WARNINGS;
-}
-
 void UMovieScene::PreSave(FObjectPreSaveContext ObjectSaveContext)
 {
 	Super::PreSave(ObjectSaveContext);
diff --git a/Engine/Source/Runtime/MovieScene/Private/MovieSceneSequence.cpp b/Engine/Source/Runtime/MovieScene/Private/MovieSceneSequence.cpp
--- a/Engine/Source/Runtime/MovieScene/Private/MovieSceneSequence.cpp
+++ b/Engine/Source/Runtime/MovieScene/Private/MovieSceneSequence.cpp
@@ -101,13 +101,6 @@ EMovieSceneServerClientMask UMovieSceneSequence::OverrideNetworkMask(EMovieScene
 	return InDefaultMask;
 }
 
-void UMovieSceneSequence::PreSave(const ITargetPlatform* TargetPlatform)
-{
-	PRAGMA_DISABLE_DEPRECATION_WARNINGS;
-	Super::PreSave(TargetPlatform);
-	PRAGMA_ENABLE_DEPRECATION_WARNINGS;
-}
-
 void UMovieSceneSequence::PreSave(FObjectPreSaveContext ObjectSaveContext)
 {
 #if WITH_EDITOR
diff --git a/Engine/Source/Runtime/MovieScene/Public/Compilation/IMovieSceneTrackTemplateProducer.h b/Engine/Source/Runtime/MovieScene/Public/Compilation/IMovieSceneTrackTemplateProducer.h
--- a/Engine/Source/Runtime/MovieScene/Public/Compilation/IMovieSceneTrackTemplateProducer.h
+++ b/Engine/Source/Runtime/MovieScene/Public/Compilation/IMovieSceneTrackTemplateProducer.h
@@ -40,13 +40,7 @@ struct FMovieSceneTrackCompilerArgs
 
 	struct FMovieSceneSequenceTemplateStore {};
 
-	UE_DEPRECATED(4.19, "Template store is no longer supplied as part of track compilation")
-	FMovieSceneSequenceTemplateStore SubSequenceStore;
-
 	struct FMovieSceneTrackCompilationParams { 	bool bForEditorPreview, bDuringBlueprintCompile; };
-
-	UE_DEPRECATED(4.19, "Template store is no longer supplied as part of track compilation")
-	FMovieSceneTrackCompilationParams Params;
 };
 
 
diff --git a/Engine/Source/Runtime/MovieScene/Public/EntitySystem/MovieSceneComponentPtr.h b/Engine/Source/Runtime/MovieScene/Public/EntitySystem/MovieSceneComponentPtr.h
--- a/Engine/Source/Runtime/MovieScene/Public/EntitySystem/MovieSceneComponentPtr.h
+++ b/Engine/Source/Runtime/MovieScene/Public/EntitySystem/MovieSceneComponentPtr.h
@@ -188,18 +188,6 @@ struct TRead : TComponentPtr<const T>
 	{
 		return (*this)[Index];
 	}
-
-	UE_DEPRECATED(4.27, "It is no longer necessary to call Resolve - please change to using the TComponentPtr API.")
-	const T* Resolve(const FEntityAllocation* Allocation) const
-	{
-		return this->AsPtr();
-	}
-
-	UE_DEPRECATED(4.27, "It is no longer necessary to call ResolveAsArray - please change to TComponentPtr::AsArray.")
-	TArrayView<const T> ResolveAsArray(const FEntityAllocation* Allocation) const
-	{
-		return MakeArrayView(this->AsPtr(), Allocation->Num());
-	}
 };
 using FReadEntityIDs = TRead<FMovieSceneEntityID>;
 
@@ -214,19 +202,6 @@ struct TReadOptional : TComponentPtr<const T>
 			this->ComponentPtr = reinterpret_cast<const T*>(InHeader->Components) + ComponentOffset;
 		}
 	}
-
-	UE_DEPRECATED(4.27, "It is no longer necessary to call Resolve - please change to using the TComponentPtr API.")
-	const T* Resolve(const FEntityAllocation* Allocation) const
-	{
-		return this->AsPtr();
-	}
-
-	UE_DEPRECATED(4.27, "It is no longer necessary to call ResolveAsArray - please change to TComponentPtr::AsArray.")
-	TArrayView<const T> ResolveAsArray(const FEntityAllocation* Allocation) const
-	{
-		const T* Ptr = this->AsPtr();
-		return MakeArrayView(Ptr, Ptr ? Allocation->Num() : 0);
-	}
 };
 
 
@@ -278,18 +253,6 @@ struct TWrite : TComponentPtr<T>
 	explicit TWrite(const FComponentHeader* InHeader, int32 ComponentOffset = 0)
 		: TComponentPtr<T>(reinterpret_cast<T*>(InHeader->Components) + ComponentOffset)
 	{}
-
-	UE_DEPRECATED(4.27, "It is no longer necessary to call Resolve - please change to using the TComponentPtr API.")
-	T* Resolve(const FEntityAllocation* Allocation) const
-	{
-		return this->AsPtr();
-	}
-
-	UE_DEPRECATED(4.27, "It is no longer necessary to call ResolveAsArray - please change to TComponentPtr::AsArray.")
-	TArrayView<T> ResolveAsArray(const FEntityAllocation* Allocation) const
-	{
-		return MakeArrayView(this->AsPtr(), Allocation->Num());
-	}
 };
 
 template<typename T>
@@ -303,19 +266,6 @@ struct TWriteOptional : TComponentPtr<T>
 			this->ComponentPtr = reinterpret_cast<T*>(InHeader->Components) + ComponentOffset;
 		}
 	}
-
-	UE_DEPRECATED(4.27, "It is no longer necessary to call Resolve - please change to using the TComponentPtr API.")
-	T* Resolve(const FEntityAllocation* Allocation) const
-	{
-		return this->AsPtr();
-	}
-
-	UE_DEPRECATED(4.27, "It is no longer necessary to call ResolveAsArray - please change to TComponentPtr::AsArray.")
-	TArrayView<T> ResolveAsArray(const FEntityAllocation* Allocation) const
-	{
-		T* Ptr = this->AsPtr();
-		return MakeArrayView(Ptr, Ptr ? Allocation->Num() : 0);
-	}
 };
 
 
diff --git a/Engine/Source/Runtime/MovieScene/Public/EntitySystem/MovieScenePreAnimatedStateSystem.h b/Engine/Source/Runtime/MovieScene/Public/EntitySystem/MovieScenePreAnimatedStateSystem.h
--- a/Engine/Source/Runtime/MovieScene/Public/EntitySystem/MovieScenePreAnimatedStateSystem.h
+++ b/Engine/Source/Runtime/MovieScene/Public/EntitySystem/MovieScenePreAnimatedStateSystem.h
@@ -45,17 +45,6 @@ public:
 
 	virtual void SavePreAnimatedState(const FPreAnimationParameters& InParameters) {}
 	virtual void RestorePreAnimatedState(const FPreAnimationParameters& InParameters) {}
-
-private:
-
-	UE_DEPRECATED(4.26, "Please override the method that takes a FPreAnimationParameters")
-	virtual void SavePreAnimatedState(UE::MovieScene::FSystemTaskPrerequisites& InPrerequisites, UE::MovieScene::FSystemSubsequentTasks& Subsequents) final {}
-
-	UE_DEPRECATED(4.26, "Please override the method that takes a FPreAnimationParameters")
-	virtual void SaveGlobalPreAnimatedState(UE::MovieScene::FSystemTaskPrerequisites& InPrerequisites, UE::MovieScene::FSystemSubsequentTasks& Subsequents) final {}
-
-	UE_DEPRECATED(4.26, "Please override the method that takes a FPreAnimationParameters")
-	virtual void RestorePreAnimatedState(UE::MovieScene::FSystemTaskPrerequisites& InPrerequisites, UE::MovieScene::FSystemSubsequentTasks& Subsequents) final {}
 };
 
 
diff --git a/Engine/Source/Runtime/MovieScene/Public/MovieScene.h b/Engine/Source/Runtime/MovieScene/Public/MovieScene.h
--- a/Engine/Source/Runtime/MovieScene/Public/MovieScene.h
+++ b/Engine/Source/Runtime/MovieScene/Public/MovieScene.h
@@ -1172,10 +1172,6 @@ protected:
 
 protected:
 
-	PRAGMA_DISABLE_DEPRECATION_WARNINGS // Suppress compiler warning on override of deprecated function
-	UE_DEPRECATED(5.0, "Use version that takes FObjectPreSaveContext instead.")
-	virtual void PreSave(const class ITargetPlatform* TargetPlatform) override;
-	PRAGMA_ENABLE_DEPRECATION_WARNINGS
 	/** Called before this object is being deserialized. */
 	virtual void PreSave(FObjectPreSaveContext ObjectSaveContext) override;
 
diff --git a/Engine/Source/Runtime/MovieScene/Public/MovieSceneExecutionToken.h b/Engine/Source/Runtime/MovieScene/Public/MovieSceneExecutionToken.h
--- a/Engine/Source/Runtime/MovieScene/Public/MovieSceneExecutionToken.h
+++ b/Engine/Source/Runtime/MovieScene/Public/MovieSceneExecutionToken.h
@@ -27,9 +27,6 @@ struct IMovieScenePreAnimatedToken
 	 * @param Params Parameters for restoring state
 	 */
 	virtual void RestoreState(UObject& Object, const UE::MovieScene::FRestoreStateParams& Params) = 0;
-
-	UE_DEPRECATED(4.27, "Please use the overload that takes UE::MovieScene::FRestoreStateParams")
-	virtual void RestoreState(UObject& Object, IMovieScenePlayer& Player) {}
 };
 typedef TInlineValue<IMovieScenePreAnimatedToken, 32> IMovieScenePreAnimatedTokenPtr;
 
@@ -74,9 +71,6 @@ struct IMovieScenePreAnimatedGlobalToken
 	 * @param Params Parameters for restoring state
 	 */
 	virtual void RestoreState(const UE::MovieScene::FRestoreStateParams& Params) = 0;
-
-	UE_DEPRECATED(4.27, "Please use the overload that takes UE::MovieScene::FRestoreStateParams")
-	virtual void RestoreState(IMovieScenePlayer& Player) {}
 };
 typedef TInlineValue<IMovieScenePreAnimatedGlobalToken, 32> IMovieScenePreAnimatedGlobalTokenPtr;
 
diff --git a/Engine/Source/Runtime/MovieScene/Public/MovieSceneSection.h b/Engine/Source/Runtime/MovieScene/Public/MovieSceneSection.h
--- a/Engine/Source/Runtime/MovieScene/Public/MovieSceneSection.h
+++ b/Engine/Source/Runtime/MovieScene/Public/MovieSceneSection.h
@@ -611,9 +611,6 @@ public:
 	 */
 	virtual void OnBindingIDsUpdated(const TMap<UE::MovieScene::FFixedObjectBindingID, UE::MovieScene::FFixedObjectBindingID>& OldFixedToNewFixedMap, FMovieSceneSequenceID LocalSequenceID, const FMovieSceneSequenceHierarchy* Hierarchy, IMovieScenePlayer& Player) {}
 
-	UE_DEPRECATED(5.0, "OnBindingsUpdated has been deprecated. Use OnBindingIDsUpdated instead")
-	virtual void OnBindingsUpdated(const TMap<FGuid, FGuid>& OldGuidToNewGuidMap) { }
-
 	/** Get the referenced bindings for this section */
 	MOVIESCENE_API virtual void GetReferencedBindings(TArray<FGuid>& OutBindings) {}
 
diff --git a/Engine/Source/Runtime/MovieScene/Public/MovieSceneSequence.h b/Engine/Source/Runtime/MovieScene/Public/MovieSceneSequence.h
--- a/Engine/Source/Runtime/MovieScene/Public/MovieSceneSequence.h
+++ b/Engine/Source/Runtime/MovieScene/Public/MovieSceneSequence.h
@@ -245,10 +245,6 @@ public:
 public:
 
 	MOVIESCENE_API virtual void PostLoad() override;
-	PRAGMA_DISABLE_DEPRECATION_WARNINGS // Suppress compiler warning on override of deprecated function
-	UE_DEPRECATED(5.0, "Use version that takes FObjectPreSaveContext instead.")
-	MOVIESCENE_API virtual void PreSave(const ITargetPlatform* TargetPlatform) override;
-	PRAGMA_ENABLE_DEPRECATION_WARNINGS
 	MOVIESCENE_API virtual void PreSave(FObjectPreSaveContext ObjectSaveContext) override;
 	MOVIESCENE_API virtual void BeginDestroy() override;
 	MOVIESCENE_API virtual void PostDuplicate(bool bDuplicateForPIE) override;
diff --git a/Engine/Source/Runtime/MovieScene/Public/MovieSceneSequencePlayer.h b/Engine/Source/Runtime/MovieScene/Public/MovieSceneSequencePlayer.h
--- a/Engine/Source/Runtime/MovieScene/Public/MovieSceneSequencePlayer.h
+++ b/Engine/Source/Runtime/MovieScene/Public/MovieSceneSequencePlayer.h
@@ -330,44 +330,6 @@ public:
 	UFUNCTION(BlueprintCallable, Category="Game|Cinematic")
 	void RestoreState();
 
-public:
-
-	UE_DEPRECATED(4.26, "PlayToFrame is deprecated, use SetPlaybackPosition.")
-	UFUNCTION(BlueprintCallable, Category = "Sequencer|Player", DisplayName = "Play To (Frames)", meta=(DeprecatedFunction, DeprecationMessage="PlayToFrame is deprecated, use SetPlaybackPosition."))
-	void PlayToFrame(FFrameTime NewPosition) { SetPlaybackPosition(FMovieSceneSequencePlaybackParams(NewPosition, EUpdatePositionMethod::Play)); }
-
-	UE_DEPRECATED(4.26, "ScrubToFrame is deprecated, use SetPlaybackPosition.")
-	UFUNCTION(BlueprintCallable, Category = "Sequencer|Player", DisplayName = "Scrub To (Frames)", meta=(DeprecatedFunction, DeprecationMessage="ScrubToFrame is deprecated, use SetPlaybackPosition."))
-	void ScrubToFrame(FFrameTime NewPosition) { SetPlaybackPosition(FMovieSceneSequencePlaybackParams(NewPosition, EUpdatePositionMethod::Scrub)); }
-
-	UE_DEPRECATED(4.26, "JumpToFrame is deprecated, use SetPlaybackPosition.")
-	UFUNCTION(BlueprintCallable, Category = "Sequencer|Player", DisplayName="Jump To (Frames)", meta=(DeprecatedFunction, DeprecationMessage="JumpToFrame is deprecated, use SetPlaybackPosition."))
-	void JumpToFrame(FFrameTime NewPosition) { SetPlaybackPosition(FMovieSceneSequencePlaybackParams(NewPosition, EUpdatePositionMethod::Jump)); }
-
-	UE_DEPRECATED(4.26, "PlayToSeconds is deprecated, use SetPlaybackPosition.")
-	UFUNCTION(BlueprintCallable, Category = "Sequencer|Player", DisplayName = "Play To (Seconds)", meta=(DeprecatedFunction, DeprecationMessage="PlayToSeconds is deprecated, use SetPlaybackPosition."))
-	void PlayToSeconds(float TimeInSeconds) { SetPlaybackPosition(FMovieSceneSequencePlaybackParams(TimeInSeconds, EUpdatePositionMethod::Play)); }
-
-	UE_DEPRECATED(4.26, "ScrubToSeconds is deprecated, use SetPlaybackPosition.")
-	UFUNCTION(BlueprintCallable, Category = "Sequencer|Player", DisplayName = "Scrub To (Seconds)", meta=(DeprecatedFunction, DeprecationMessage="ScrubToSeconds is deprecated, use SetPlaybackPosition."))
-	void ScrubToSeconds(float TimeInSeconds) { SetPlaybackPosition(FMovieSceneSequencePlaybackParams(TimeInSeconds, EUpdatePositionMethod::Scrub)); }
-
-	UE_DEPRECATED(4.26, "JumpToSeconds is deprecated, use SetPlaybackPosition.")
-	UFUNCTION(BlueprintCallable, Category = "Sequencer|Player", DisplayName = "Jump To (Seconds)", meta=(DeprecatedFunction, DeprecationMessage="JumpToSeconds is deprecated, use SetPlaybackPosition."))
-	void JumpToSeconds(float TimeInSeconds) { SetPlaybackPosition(FMovieSceneSequencePlaybackParams(TimeInSeconds, EUpdatePositionMethod::Jump)); }
-
-	UE_DEPRECATED(4.26, "PlayToMarkedFrame is deprecated, use SetPlaybackPosition.")
-	UFUNCTION(BlueprintCallable, Category = "Sequencer|Player", meta=(DeprecatedFunction, DeprecationMessage="PlayToMarkedFrame is deprecated, use SetPlaybackPosition."))
-	bool PlayToMarkedFrame(const FString& InLabel) { SetPlaybackPosition(FMovieSceneSequencePlaybackParams(InLabel, EUpdatePositionMethod::Play)); return true; }
-
-	UE_DEPRECATED(4.26, "ScrubToMarkedFrame is deprecated, use SetPlaybackPosition.")
-	UFUNCTION(BlueprintCallable, Category = "Sequencer|Player", meta=(DeprecatedFunction, DeprecationMessage="ScrubToMarkedFrame is deprecated, use SetPlaybackPosition."))
-	bool ScrubToMarkedFrame(const FString& InLabel) { SetPlaybackPosition(FMovieSceneSequencePlaybackParams(InLabel, EUpdatePositionMethod::Scrub)); return true; }
-
-	UE_DEPRECATED(4.26, "JumpToMarkedFrame is deprecated, use SetPlaybackPosition.")
-	UFUNCTION(BlueprintCallable, Category = "Sequencer|Player", meta=(DeprecatedFunction, DeprecationMessage="JumpToMarkedFrame is deprecated, use SetPlaybackPosition."))
-	bool JumpToMarkedFrame(const FString& InLabel) { SetPlaybackPosition(FMovieSceneSequencePlaybackParams(InLabel, EUpdatePositionMethod::Jump)); return true; }
-
 public:
 
 	/** Check whether the sequence is actively playing. */
