Sequencer: Use a transient player to get the bound objects which will properly return subobjects

diff --git a/Engine/Plugins/MovieScene/SequencerScripting/Source/SequencerScripting/Private/ExtensionLibraries/MovieSceneSequenceExtensions.cpp b/Engine/Plugins/MovieScene/SequencerScripting/Source/SequencerScripting/Private/ExtensionLibraries/MovieSceneSequenceExtensions.cpp
--- a/Engine/Plugins/MovieScene/SequencerScripting/Source/SequencerScripting/Private/ExtensionLibraries/MovieSceneSequenceExtensions.cpp
+++ b/Engine/Plugins/MovieScene/SequencerScripting/Source/SequencerScripting/Private/ExtensionLibraries/MovieSceneSequenceExtensions.cpp
@@ -7,6 +7,7 @@
 #include "MovieSceneSequence.h"
 #include "MovieScene.h"
 #include "MovieSceneFolder.h"
+#include "Evaluation/MovieSceneEvaluationTemplateInstance.h"
 #include "Compilation/MovieSceneCompiledDataManager.h"
 #include "MovieSceneSpawnable.h"
 #include "SequencerScriptingRange.h"
@@ -820,11 +821,29 @@ TArray<UObject*> UMovieSceneSequenceExtensions::LocateBoundObjects(UMovieSceneSe
 		return TArray<UObject*>();
 	}
 
-	TArray<UObject*> Result;
-	TArray<UObject*, TInlineAllocator<1>> OutObjects;
-	Sequence->LocateBoundObjects(InBinding.BindingID, Context, OutObjects);
-	Result.Append(OutObjects);
+	class FTransientPlayer : public IMovieScenePlayer
+	{
+	public:
+		FMovieSceneRootEvaluationTemplateInstance Template;
+		virtual FMovieSceneRootEvaluationTemplateInstance& GetEvaluationTemplate() override { check(false); return Template; }
+		virtual void UpdateCameraCut(UObject* CameraObject, const EMovieSceneCameraCutParams& CameraCutParams) override {}
+		virtual void SetViewportSettings(const TMap<FViewportClient*, EMovieSceneViewportParams>& ViewportParamsMap) override {}
+		virtual void GetViewportSettings(TMap<FViewportClient*, EMovieSceneViewportParams>& ViewportParamsMap) const override {}
+		virtual EMovieScenePlayerStatus::Type GetPlaybackStatus() const { return EMovieScenePlayerStatus::Stopped; }
+		virtual void SetPlaybackStatus(EMovieScenePlayerStatus::Type InPlaybackStatus) override {}
+	} Player;
 
+	Player.State.AssignSequence(MovieSceneSequenceID::Root, *Sequence, Player);
+
+	TArrayView<TWeakObjectPtr<>> Objects = Player.FindBoundObjects(InBinding.BindingID, MovieSceneSequenceID::Root);
+	TArray<UObject*> Result;
+	for (TWeakObjectPtr<> WeakObject : Objects)
+	{
+		if (WeakObject.IsValid())
+		{
+			Result.Add(WeakObject.Get());
+		}
+	}
 	return Result;
 }
 
