Sequencer: Deprecate GetBoundObjects because it doesn't work with spawnables. Better to use ULevelSequenceEditorBlueprintLibrary::GetBoundObjects because it operates on the currently opened sequence in Editor and can return valid spawnable objects.

Also deprecated the old movie scene capture functions.

diff --git a/Engine/Plugins/MovieScene/SequencerScripting/Source/SequencerScriptingEditor/Private/SequencerTools.cpp b/Engine/Plugins/MovieScene/SequencerScripting/Source/SequencerScriptingEditor/Private/SequencerTools.cpp
--- a/Engine/Plugins/MovieScene/SequencerScripting/Source/SequencerScriptingEditor/Private/SequencerTools.cpp
+++ b/Engine/Plugins/MovieScene/SequencerScripting/Source/SequencerScriptingEditor/Private/SequencerTools.cpp
@@ -51,11 +51,13 @@ bool USequencerToolsFunctionLibrary::RenderMovie(UMovieSceneCapture* InCaptureSe
 		return false;
 	}
 
+PRAGMA_DISABLE_DEPRECATION_WARNINGS
 	if (IsRenderingMovie())
 	{
 		FFrame::KismetExecutionMessage(TEXT("Capture already in progress."), ELogVerbosity::Error);
 		return false;
 	}
+PRAGMA_ENABLE_DEPRECATION_WARNINGS
 
 	// If they're capturing a level sequence we'll do some additional checking as there are more parameters on the Automated Level Sequence capture.
 	UAutomatedLevelSequenceCapture* LevelSequenceCapture = Cast<UAutomatedLevelSequenceCapture>(InCaptureSettings);
diff --git a/Engine/Plugins/MovieScene/SequencerScripting/Source/SequencerScriptingEditor/Public/SequencerTools.h b/Engine/Plugins/MovieScene/SequencerScripting/Source/SequencerScriptingEditor/Public/SequencerTools.h
--- a/Engine/Plugins/MovieScene/SequencerScripting/Source/SequencerScriptingEditor/Public/SequencerTools.h
+++ b/Engine/Plugins/MovieScene/SequencerScripting/Source/SequencerScriptingEditor/Public/SequencerTools.h
@@ -122,13 +122,15 @@ public:
 	* if we're rendering via a PIE instance or a new process based on the passed in settings. Will return false
 	* if the state is not valid (ie: null or missing required parameters, capture in progress, etc.), true otherwise.
 	*/
-	UFUNCTION(BlueprintCallable, Category = "Editor Scripting | Sequencer Tools | Movie Rendering")
+	UE_DEPRECATED(5.3, "SequencerTools::RenderMovie has been deprecated, please use Movie Render Queue")
+	UFUNCTION(BlueprintCallable, Category = "Editor Scripting | Sequencer Tools | Movie Rendering", meta=(DeprecatedFunction))
 	static bool RenderMovie(class UMovieSceneCapture* InCaptureSettings, FOnRenderMovieStopped OnFinishedCallback);
 
 	/** 
 	* Returns if Render to Movie is currently in progress.
 	*/
-	UFUNCTION(BlueprintCallable, Category = "Editor Scripting | Sequencer Tools | Movie Rendering")
+	UE_DEPRECATED(5.3, "SequencerTools::IsRenderingMovie has been deprecated, please use Movie Render Queue")
+	UFUNCTION(BlueprintCallable, Category = "Editor Scripting | Sequencer Tools | Movie Rendering", meta=(DeprecatedFunction))
 	static bool IsRenderingMovie()
 	{
 		IMovieSceneCaptureDialogModule& MovieSceneCaptureModule = FModuleManager::Get().LoadModuleChecked<IMovieSceneCaptureDialogModule>("MovieSceneCaptureDialog");
@@ -138,7 +140,8 @@ public:
 	/**
 	* Attempts to cancel an in-progress Render to Movie. Does nothing if there is no render in progress.
 	*/
-	UFUNCTION(BlueprintCallable, Category = "Editor Scripting | Sequencer Tools | Movie Rendering")
+	UE_DEPRECATED(5.3, "SequencerTools::CancelMovieRender has been deprecated, please use Movie Render Queue")
+	UFUNCTION(BlueprintCallable, Category = "Editor Scripting | Sequencer Tools | Movie Rendering", meta=(DeprecatedFunction))
 	static void CancelMovieRender();
 
 public:
@@ -147,14 +150,16 @@ public:
 	 * Retrieve all objects currently bound to the specified binding identifiers. The sequence will be evaluated in lower bound of the specified range, 
 	 * which allows for retrieving spawnables in that period of time.
 	 */
-	UFUNCTION(BlueprintCallable, Category = "Editor Scripting | Sequencer Tools")
+	UE_DEPRECATED(5.3, "SequencerTools::GetBoundObjects has been deprecated because it doesn't work with Spawnables, please use ULevelSequenceEditorBlueprintLibrary::GetBoundObjects which operates on the currently opened sequence in Editor")
+	UFUNCTION(BlueprintCallable, Category = "Editor Scripting | Sequencer Tools", meta=(DeprecatedFunction))
 	static TArray<FSequencerBoundObjects> GetBoundObjects(UWorld* InWorld, ULevelSequence* InSequence, const TArray<FMovieSceneBindingProxy>& InBindings, const FSequencerScriptingRange& InRange);
 
 	/*
 	 * Get the object bindings for the requested object. The sequence will be evaluated in lower bound of the specified range, 
 	 * which allows for retrieving spawnables in that period of time.
 	 */
-	UFUNCTION(BlueprintCallable, Category = "Editor Scripting | Sequencer Tools")
+	UE_DEPRECATED(5.3, "SequencerTools::GetObjectBindings has been deprecated because it doesn't work with Spawnables, please use ULevelSequenceEditorBlueprintLibrary::GetObjectBindings which operates on the currently opened sequence in Editor")
+	UFUNCTION(BlueprintCallable, Category = "Editor Scripting | Sequencer Tools", meta=(DeprecatedFunction))
 	static TArray<FSequencerBoundObjects> GetObjectBindings(UWorld* InWorld, ULevelSequence* InSequence, const TArray<UObject*>& InObject, const FSequencerScriptingRange& InRange);
 
 public:
