PR #10370: MovieSceneSequenceExtensions::GetBindingId not exposed to script (Contributed by [REDACTED_PERSON])

diff --git a/Engine/Plugins/MovieScene/SequencerScripting/Source/SequencerScripting/Private/ExtensionLibraries/MovieSceneSequenceExtensions.cpp b/Engine/Plugins/MovieScene/SequencerScripting/Source/SequencerScripting/Private/ExtensionLibraries/MovieSceneSequenceExtensions.cpp
--- a/Engine/Plugins/MovieScene/SequencerScripting/Source/SequencerScripting/Private/ExtensionLibraries/MovieSceneSequenceExtensions.cpp
+++ b/Engine/Plugins/MovieScene/SequencerScripting/Source/SequencerScripting/Private/ExtensionLibraries/MovieSceneSequenceExtensions.cpp
@@ -847,7 +847,7 @@ TArray<UObject*> UMovieSceneSequenceExtensions::LocateBoundObjects(UMovieSceneSe
 	return Result;
 }
 
-FMovieSceneObjectBindingID UMovieSceneSequenceExtensions::GetBindingID(const FMovieSceneBindingProxy& InBinding)
+FMovieSceneObjectBindingID UMovieSceneSequenceExtensions::GetBindingID(UMovieSceneSequence* Sequence, const FMovieSceneBindingProxy& InBinding)
 {
 	return UE::MovieScene::FRelativeObjectBindingID(InBinding.BindingID);
 }
diff --git a/Engine/Plugins/MovieScene/SequencerScripting/Source/SequencerScripting/Public/ExtensionLibraries/MovieSceneSequenceExtensions.h b/Engine/Plugins/MovieScene/SequencerScripting/Source/SequencerScripting/Public/ExtensionLibraries/MovieSceneSequenceExtensions.h
--- a/Engine/Plugins/MovieScene/SequencerScripting/Source/SequencerScripting/Public/ExtensionLibraries/MovieSceneSequenceExtensions.h
+++ b/Engine/Plugins/MovieScene/SequencerScripting/Source/SequencerScripting/Public/ExtensionLibraries/MovieSceneSequenceExtensions.h
@@ -458,7 +458,7 @@ public:
 	 * @return The binding's id
 	 */
 	UFUNCTION(BlueprintCallable, Category = "Sequencer|Sequence", meta = (ScriptMethod))
-	static FMovieSceneObjectBindingID GetBindingID(const FMovieSceneBindingProxy& InBinding);
+	static FMovieSceneObjectBindingID GetBindingID(UMovieSceneSequence* Sequence, const FMovieSceneBindingProxy& InBinding);
 
 
 	/**
