Sequencer: Remove check(false) in FTransientPlayer's GetEvaluationTemplate() for cases where a director instance is needed

#rnx

diff --git a/Engine/Plugins/MovieScene/SequencerScripting/Source/SequencerScripting/Private/ExtensionLibraries/MovieSceneSequenceExtensions.cpp b/Engine/Plugins/MovieScene/SequencerScripting/Source/SequencerScripting/Private/ExtensionLibraries/MovieSceneSequenceExtensions.cpp
--- a/Engine/Plugins/MovieScene/SequencerScripting/Source/SequencerScripting/Private/ExtensionLibraries/MovieSceneSequenceExtensions.cpp
+++ b/Engine/Plugins/MovieScene/SequencerScripting/Source/SequencerScripting/Private/ExtensionLibraries/MovieSceneSequenceExtensions.cpp
@@ -825,7 +825,7 @@ TArray<UObject*> UMovieSceneSequenceExtensions::LocateBoundObjects(UMovieSceneSe
 	{
 	public:
 		FMovieSceneRootEvaluationTemplateInstance Template;
-		virtual FMovieSceneRootEvaluationTemplateInstance& GetEvaluationTemplate() override { check(false); return Template; }
+		virtual FMovieSceneRootEvaluationTemplateInstance& GetEvaluationTemplate() override { return Template; }
 		virtual void UpdateCameraCut(UObject* CameraObject, const EMovieSceneCameraCutParams& CameraCutParams) override {}
 		virtual void SetViewportSettings(const TMap<FViewportClient*, EMovieSceneViewportParams>& ViewportParamsMap) override {}
 		virtual void GetViewportSettings(TMap<FViewportClient*, EMovieSceneViewportParams>& ViewportParamsMap) const override {}
diff --git a/Engine/Source/Runtime/LevelSequence/Private/LevelSequence.cpp b/Engine/Source/Runtime/LevelSequence/Private/LevelSequence.cpp
--- a/Engine/Source/Runtime/LevelSequence/Private/LevelSequence.cpp
+++ b/Engine/Source/Runtime/LevelSequence/Private/LevelSequence.cpp
@@ -603,7 +603,7 @@ FGuid ULevelSequence::FindOrAddBinding(UObject* InObject)
 		{
 		public:
 			FMovieSceneRootEvaluationTemplateInstance Template;
-			virtual FMovieSceneRootEvaluationTemplateInstance& GetEvaluationTemplate() override { check(false); return Template; }
+			virtual FMovieSceneRootEvaluationTemplateInstance& GetEvaluationTemplate() override { return Template; }
 			virtual void UpdateCameraCut(UObject* CameraObject, const EMovieSceneCameraCutParams& CameraCutParams) override {}
 			virtual void SetViewportSettings(const TMap<FViewportClient*, EMovieSceneViewportParams>& ViewportParamsMap) override {}
 			virtual void GetViewportSettings(TMap<FViewportClient*, EMovieSceneViewportParams>& ViewportParamsMap) const override {}
diff --git a/Engine/Source/Runtime/MovieScene/Private/MovieSceneSequence.cpp b/Engine/Source/Runtime/MovieScene/Private/MovieSceneSequence.cpp
--- a/Engine/Source/Runtime/MovieScene/Private/MovieSceneSequence.cpp
+++ b/Engine/Source/Runtime/MovieScene/Private/MovieSceneSequence.cpp
@@ -298,7 +298,7 @@ FGuid UMovieSceneSequence::FindPossessableObjectId(UObject& Object, UObject* Con
 	{
 	public:
 		FMovieSceneRootEvaluationTemplateInstance Template;
-		virtual FMovieSceneRootEvaluationTemplateInstance& GetEvaluationTemplate() override { check(false); return Template; }
+		virtual FMovieSceneRootEvaluationTemplateInstance& GetEvaluationTemplate() override { return Template; }
 		virtual void UpdateCameraCut(UObject* CameraObject, const EMovieSceneCameraCutParams& CameraCutParams) override {}
 		virtual void SetViewportSettings(const TMap<FViewportClient*, EMovieSceneViewportParams>& ViewportParamsMap) override {}
 		virtual void GetViewportSettings(TMap<FViewportClient*, EMovieSceneViewportParams>& ViewportParamsMap) const override {}
