Sequencer Scripting: Change to get/set view and working ranges to double

diff --git a/Engine/Plugins/Animation/AnimationData/Source/AnimationData/Private/AnimSequencerController.cpp b/Engine/Plugins/Animation/AnimationData/Source/AnimationData/Private/AnimSequencerController.cpp
--- a/Engine/Plugins/Animation/AnimationData/Source/AnimationData/Private/AnimSequencerController.cpp
+++ b/Engine/Plugins/Animation/AnimationData/Source/AnimationData/Private/AnimSequencerController.cpp
@@ -230,10 +230,10 @@ void UAnimSequencerController::SetMovieSceneRange(FFrameNumber InFrameNumber) co
 	MovieScene->SetPlaybackRangeLocked(false);
 	
 	FMovieSceneEditorData& EditorData = MovieScene->GetEditorData();
-	EditorData.ViewStart = EditorData.WorkStart = 0.f;
+	EditorData.ViewStart = EditorData.WorkStart = 0.0;
 	EditorData.ViewEnd = EditorData.WorkEnd = Model->GetFrameRate().AsSeconds(InFrameNumber);
-	MovieScene->SetWorkingRange(static_cast<float>(EditorData.ViewStart), static_cast<float>(EditorData.ViewEnd));
-	MovieScene->SetViewRange(static_cast<float>(EditorData.ViewStart), static_cast<float>(EditorData.ViewEnd));
+	MovieScene->SetWorkingRange(EditorData.ViewStart, EditorData.ViewEnd);
+	MovieScene->SetViewRange(EditorData.ViewStart, EditorData.ViewEnd);
 
 	if (UMovieSceneControlRigParameterSection* Section = Model->GetFKControlRigSection())
 	{
diff --git a/Engine/Plugins/MovieScene/SequencerScripting/Source/SequencerScripting/Private/ExtensionLibraries/MovieSceneSequenceExtensions.cpp b/Engine/Plugins/MovieScene/SequencerScripting/Source/SequencerScripting/Private/ExtensionLibraries/MovieSceneSequenceExtensions.cpp
--- a/Engine/Plugins/MovieScene/SequencerScripting/Source/SequencerScripting/Private/ExtensionLibraries/MovieSceneSequenceExtensions.cpp
+++ b/Engine/Plugins/MovieScene/SequencerScripting/Source/SequencerScripting/Private/ExtensionLibraries/MovieSceneSequenceExtensions.cpp
@@ -445,7 +445,7 @@ void UMovieSceneSequenceExtensions::SetPlaybackEndSeconds(UMovieSceneSequence* S
 	}
 }
 
-void UMovieSceneSequenceExtensions::SetViewRangeStart(UMovieSceneSequence* Sequence, float StartTimeInSeconds)
+void UMovieSceneSequenceExtensions::SetViewRangeStart(UMovieSceneSequence* Sequence, double StartTimeInSeconds)
 {
 	if (!Sequence)
 	{
@@ -462,12 +462,12 @@ void UMovieSceneSequenceExtensions::SetViewRangeStart(UMovieSceneSequence* Seque
 	}
 }
 
-float UMovieSceneSequenceExtensions::GetViewRangeStart(UMovieSceneSequence* Sequence)
+double UMovieSceneSequenceExtensions::GetViewRangeStart(UMovieSceneSequence* Sequence)
 {
 	if (!Sequence)
 	{
 		FFrame::KismetExecutionMessage(TEXT("Cannot call GetViewRangeStart on a null sequence"), ELogVerbosity::Error);
-		return 0.f;
+		return 0;
 	}
 
 	UMovieScene* MovieScene = GetMovieScene(Sequence);
@@ -477,10 +477,10 @@ float UMovieSceneSequenceExtensions::GetViewRangeStart(UMovieSceneSequence* Sequ
 		return MovieScene->GetEditorData().ViewStart;
 #endif
 	}
-	return 0.f;
+	return 0;
 }
 
-void UMovieSceneSequenceExtensions::SetViewRangeEnd(UMovieSceneSequence* Sequence, float EndTimeInSeconds)
+void UMovieSceneSequenceExtensions::SetViewRangeEnd(UMovieSceneSequence* Sequence, double EndTimeInSeconds)
 {
 	if (!Sequence)
 	{
@@ -497,12 +497,12 @@ void UMovieSceneSequenceExtensions::SetViewRangeEnd(UMovieSceneSequence* Sequenc
 	}
 }
 
-float UMovieSceneSequenceExtensions::GetViewRangeEnd(UMovieSceneSequence* Sequence)
+double UMovieSceneSequenceExtensions::GetViewRangeEnd(UMovieSceneSequence* Sequence)
 {
 	if (!Sequence)
 	{
 		FFrame::KismetExecutionMessage(TEXT("Cannot call GetViewRangeEnd on a null sequence"), ELogVerbosity::Error);
-		return 0.f;
+		return 0;
 	}
 
 	UMovieScene* MovieScene = GetMovieScene(Sequence);
@@ -512,10 +512,10 @@ float UMovieSceneSequenceExtensions::GetViewRangeEnd(UMovieSceneSequence* Sequen
 		return MovieScene->GetEditorData().ViewEnd;
 #endif
 	}
-	return 0.f;
+	return 0;
 }
 
-void UMovieSceneSequenceExtensions::SetWorkRangeStart(UMovieSceneSequence* Sequence, float StartTimeInSeconds)
+void UMovieSceneSequenceExtensions::SetWorkRangeStart(UMovieSceneSequence* Sequence, double StartTimeInSeconds)
 {
 	if (!Sequence)
 	{
@@ -532,12 +532,12 @@ void UMovieSceneSequenceExtensions::SetWorkRangeStart(UMovieSceneSequence* Seque
 	}
 }
 
-float UMovieSceneSequenceExtensions::GetWorkRangeStart(UMovieSceneSequence* Sequence)
+double UMovieSceneSequenceExtensions::GetWorkRangeStart(UMovieSceneSequence* Sequence)
 {
 	if (!Sequence)
 	{
 		FFrame::KismetExecutionMessage(TEXT("Cannot call GetWorkRangeStart on a null sequence"), ELogVerbosity::Error);
-		return 0.f;
+		return 0;
 	}
 
 	UMovieScene* MovieScene = GetMovieScene(Sequence);
@@ -547,10 +547,10 @@ float UMovieSceneSequenceExtensions::GetWorkRangeStart(UMovieSceneSequence* Sequ
 		return MovieScene->GetEditorData().WorkStart;
 #endif
 	}
-	return 0.f;
+	return 0;
 }
 
-void UMovieSceneSequenceExtensions::SetWorkRangeEnd(UMovieSceneSequence* Sequence, float EndTimeInSeconds)
+void UMovieSceneSequenceExtensions::SetWorkRangeEnd(UMovieSceneSequence* Sequence, double EndTimeInSeconds)
 {
 	if (!Sequence)
 	{
@@ -567,12 +567,12 @@ void UMovieSceneSequenceExtensions::SetWorkRangeEnd(UMovieSceneSequence* Sequenc
 	}
 }
 
-float UMovieSceneSequenceExtensions::GetWorkRangeEnd(UMovieSceneSequence* Sequence)
+double UMovieSceneSequenceExtensions::GetWorkRangeEnd(UMovieSceneSequence* Sequence)
 {
 	if (!Sequence)
 	{
 		FFrame::KismetExecutionMessage(TEXT("Cannot call GetWorkRangeEnd on a null sequence"), ELogVerbosity::Error);
-		return 0.f;
+		return 0;
 	}
 
 	UMovieScene* MovieScene = GetMovieScene(Sequence);
@@ -582,7 +582,7 @@ float UMovieSceneSequenceExtensions::GetWorkRangeEnd(UMovieSceneSequence* Sequen
 		return MovieScene->GetEditorData().WorkEnd;
 #endif
 	}
-	return 0.f;
+	return 0;
 }
 
 void UMovieSceneSequenceExtensions::SetEvaluationType(UMovieSceneSequence* Sequence, EMovieSceneEvaluationType InEvaluationType)
diff --git a/Engine/Plugins/MovieScene/SequencerScripting/Source/SequencerScripting/Public/ExtensionLibraries/MovieSceneSequenceExtensions.h b/Engine/Plugins/MovieScene/SequencerScripting/Source/SequencerScripting/Public/ExtensionLibraries/MovieSceneSequenceExtensions.h
--- a/Engine/Plugins/MovieScene/SequencerScripting/Source/SequencerScripting/Public/ExtensionLibraries/MovieSceneSequenceExtensions.h
+++ b/Engine/Plugins/MovieScene/SequencerScripting/Source/SequencerScripting/Public/ExtensionLibraries/MovieSceneSequenceExtensions.h
@@ -260,7 +260,7 @@ public:
 	 * @param StartTimeInSeconds The desired view range start time in seconds for this sequence
 	 */
 	UFUNCTION(BlueprintCallable, Category = "Sequencer|Sequence", meta = (ScriptMethod, DevelopmentOnly))
-	static void SetViewRangeStart(UMovieSceneSequence* InSequence, float StartTimeInSeconds);
+	static void SetViewRangeStart(UMovieSceneSequence* InSequence, double StartTimeInSeconds);
 
 	/**
 	 * Get the sequence view range start in seconds
@@ -269,7 +269,7 @@ public:
 	 * @return The view range start time in seconds for this sequence
 	 */
 	UFUNCTION(BlueprintPure, Category = "Sequencer|Sequence", meta = (ScriptMethod, DevelopmentOnly))
-	static float GetViewRangeStart(UMovieSceneSequence* InSequence);
+	static double GetViewRangeStart(UMovieSceneSequence* InSequence);
 
 	/**
 	 * Set the sequence view range end in seconds
@@ -278,7 +278,7 @@ public:
 	 * @param StartTimeInSeconds The desired view range end time in seconds for this sequence
 	 */
 	UFUNCTION(BlueprintCallable, Category = "Sequencer|Sequence", meta = (ScriptMethod, DevelopmentOnly))
-	static void SetViewRangeEnd(UMovieSceneSequence* InSequence, float EndTimeInSeconds);
+	static void SetViewRangeEnd(UMovieSceneSequence* InSequence, double EndTimeInSeconds);
 
 	/**
 	 * Get the sequence view range end in seconds
@@ -287,7 +287,7 @@ public:
 	 * @return The view range end time in seconds for this sequence
 	 */
 	UFUNCTION(BlueprintPure, Category = "Sequencer|Sequence", meta = (ScriptMethod, DevelopmentOnly))
-	static float GetViewRangeEnd(UMovieSceneSequence* InSequence);
+	static double GetViewRangeEnd(UMovieSceneSequence* InSequence);
 
 	/**
 	 * Set the sequence work range start in seconds
@@ -296,7 +296,7 @@ public:
 	 * @param StartTimeInSeconds The desired work range start time in seconds for this sequence
 	 */
 	UFUNCTION(BlueprintCallable, Category = "Sequencer|Sequence", meta = (ScriptMethod, DevelopmentOnly))
-	static void SetWorkRangeStart(UMovieSceneSequence* InSequence, float StartTimeInSeconds);
+	static void SetWorkRangeStart(UMovieSceneSequence* InSequence, double StartTimeInSeconds);
 
 	/**
 	 * Get the sequence work range start in seconds
@@ -305,7 +305,7 @@ public:
 	 * @return The work range start time in seconds for this sequence
 	 */
 	UFUNCTION(BlueprintPure, Category = "Sequencer|Sequence", meta = (ScriptMethod, DevelopmentOnly))
-	static float GetWorkRangeStart(UMovieSceneSequence* InSequence);
+	static double GetWorkRangeStart(UMovieSceneSequence* InSequence);
 
 	/**
 	 * Set the sequence work range end in seconds
@@ -314,7 +314,7 @@ public:
 	 * @param StartTimeInSeconds The desired work range end time in seconds for this sequence
 	 */
 	UFUNCTION(BlueprintCallable, Category = "Sequencer|Sequence", meta = (ScriptMethod, DevelopmentOnly))
-	static void SetWorkRangeEnd(UMovieSceneSequence* InSequence, float EndTimeInSeconds);
+	static void SetWorkRangeEnd(UMovieSceneSequence* InSequence, double EndTimeInSeconds);
 
 	/**
 	 * Get the sequence work range end in seconds
@@ -323,7 +323,7 @@ public:
 	 * @return The work range end time in seconds for this sequence
 	 */
 	UFUNCTION(BlueprintPure, Category = "Sequencer|Sequence", meta = (ScriptMethod, DevelopmentOnly))
-	static float GetWorkRangeEnd(UMovieSceneSequence* InSequence);
+	static double GetWorkRangeEnd(UMovieSceneSequence* InSequence);
 
 	/**
 	 * Set the evaluation type for this sequence
diff --git a/Engine/Source/Runtime/MovieScene/Private/MovieScene.cpp b/Engine/Source/Runtime/MovieScene/Private/MovieScene.cpp
--- a/Engine/Source/Runtime/MovieScene/Private/MovieScene.cpp
+++ b/Engine/Source/Runtime/MovieScene/Private/MovieScene.cpp
@@ -626,7 +626,7 @@ void UMovieScene::SetPlaybackRange(const TRange<FFrameNumber>& NewRange, bool bA
 #endif
 }
 
-void UMovieScene::SetWorkingRange(float Start, float End)
+void UMovieScene::SetWorkingRange(double Start, double End)
 {
 #if WITH_EDITORONLY_DATA
 	EditorData.WorkStart = Start;
@@ -634,7 +634,7 @@ void UMovieScene::SetWorkingRange(float Start, float End)
 #endif
 }
 
-void UMovieScene::SetViewRange(float Start, float End)
+void UMovieScene::SetViewRange(double Start, double End)
 {
 #if WITH_EDITORONLY_DATA
 	EditorData.ViewStart = Start;
diff --git a/Engine/Source/Runtime/MovieScene/Public/MovieScene.h b/Engine/Source/Runtime/MovieScene/Public/MovieScene.h
--- a/Engine/Source/Runtime/MovieScene/Public/MovieScene.h
+++ b/Engine/Source/Runtime/MovieScene/Public/MovieScene.h
@@ -987,7 +987,7 @@ public:
 	 * @param Start The offset from 0-time to view this movie scene.
 	 * @param End The offset from 0-time to view this movie scene
 	 */
-	MOVIESCENE_API void SetWorkingRange(float Start, float End);
+	MOVIESCENE_API void SetWorkingRange(double Start, double End);
 
 	/**
 	 * Set the start and end view range (inner) for this movie scene
@@ -995,7 +995,7 @@ public:
 	 * @param Start The offset from 0-time to view this movie scene
 	 * @param End The offset from 0-time to view this movie scene
 	 */
-	MOVIESCENE_API void SetViewRange(float Start, float End);
+	MOVIESCENE_API void SetViewRange(double Start, double End);
 
 #if WITH_EDITORONLY_DATA
 
