Sequencer: Remove dependency on SequencerScripting private

#rnx

diff --git a/Engine/Plugins/MovieScene/SequencerScripting/Source/SequencerScriptingEditor/SequencerScriptingEditor.Build.cs b/Engine/Plugins/MovieScene/SequencerScripting/Source/SequencerScriptingEditor/SequencerScriptingEditor.Build.cs
--- a/Engine/Plugins/MovieScene/SequencerScripting/Source/SequencerScriptingEditor/SequencerScriptingEditor.Build.cs
+++ b/Engine/Plugins/MovieScene/SequencerScripting/Source/SequencerScriptingEditor/SequencerScriptingEditor.Build.cs
@@ -15,7 +15,6 @@ public class SequencerScriptingEditor : ModuleRules
 
 		PrivateIncludePaths.AddRange(
 			new string[] {
-                "SequencerScripting/Private",
                 "../../../../Source/Editor/UnrealEd/Private", // TODO: Fix this, for now it's needed for the fbx exporter
             }
         );
