Sequencer: Sequencer and Curve Editor Key Selection

diff --git a/Engine/Plugins/Animation/ControlRig/Source/ControlRig/Private/Sequencer/MovieSceneControlRigSpaceChannel.cpp b/Engine/Plugins/Animation/ControlRig/Source/ControlRig/Private/Sequencer/MovieSceneControlRigSpaceChannel.cpp
--- a/Engine/Plugins/Animation/ControlRig/Source/ControlRig/Private/Sequencer/MovieSceneControlRigSpaceChannel.cpp
+++ b/Engine/Plugins/Animation/ControlRig/Source/ControlRig/Private/Sequencer/MovieSceneControlRigSpaceChannel.cpp
@@ -27,6 +27,16 @@ void FMovieSceneControlRigSpaceChannel::GetKeys(const TRange<FFrameNumber>& With
 	GetData().GetKeys(WithinRange, OutKeyTimes, DeleteKeyHandles);
 }
 
+FKeyHandle FMovieSceneControlRigSpaceChannel::GetHandle(int32 Index)
+{
+	return GetData().GetHandle(Index);
+}
+
+int32 FMovieSceneControlRigSpaceChannel::GetIndex(FKeyHandle Handle)
+{
+	return GetData().GetIndex(Handle);
+}
+
 void FMovieSceneControlRigSpaceChannel::GetKeyTimes(TArrayView<const FKeyHandle> InHandles, TArrayView<FFrameNumber> OutKeyTimes)
 {
 	GetData().GetKeyTimes(InHandles, OutKeyTimes);
diff --git a/Engine/Plugins/Animation/ControlRig/Source/ControlRig/Public/Sequencer/MovieSceneControlRigSpaceChannel.h b/Engine/Plugins/Animation/ControlRig/Source/ControlRig/Public/Sequencer/MovieSceneControlRigSpaceChannel.h
--- a/Engine/Plugins/Animation/ControlRig/Source/ControlRig/Public/Sequencer/MovieSceneControlRigSpaceChannel.h
+++ b/Engine/Plugins/Animation/ControlRig/Source/ControlRig/Public/Sequencer/MovieSceneControlRigSpaceChannel.h
@@ -111,6 +111,8 @@ public:
 	virtual int32 GetNumKeys() const override;
 	virtual void Reset() override;
 	virtual void Offset(FFrameNumber DeltaPosition) override;
+	virtual FKeyHandle GetHandle(int32 Index) override;
+	virtual int32 GetIndex(FKeyHandle Handle) override;
 
 	void GetUniqueSpaceList(TArray<FRigElementKey>* OutList);
 	FMovieSceneControlRigSpaceChannelSpaceNoLongerUsedEvent& OnSpaceNoLongerUsed() { return SpaceNoLongerUsedEvent; }
diff --git a/Engine/Plugins/MovieScene/LevelSequenceEditor/Source/LevelSequenceEditor/LevelSequenceEditor.Build.cs b/Engine/Plugins/MovieScene/LevelSequenceEditor/Source/LevelSequenceEditor/LevelSequenceEditor.Build.cs
--- a/Engine/Plugins/MovieScene/LevelSequenceEditor/Source/LevelSequenceEditor/LevelSequenceEditor.Build.cs
+++ b/Engine/Plugins/MovieScene/LevelSequenceEditor/Source/LevelSequenceEditor/LevelSequenceEditor.Build.cs
@@ -22,6 +22,7 @@ public class LevelSequenceEditor : ModuleRules
 		PublicDependencyModuleNames.AddRange(
 			new string[] {
 				"SequencerScripting",
+				"SequencerScriptingEditor",
 			}
 		);
         
diff --git a/Engine/Plugins/MovieScene/LevelSequenceEditor/Source/LevelSequenceEditor/Private/LevelSequenceEditorBlueprintLibrary.cpp b/Engine/Plugins/MovieScene/LevelSequenceEditor/Source/LevelSequenceEditor/Private/LevelSequenceEditorBlueprintLibrary.cpp
--- a/Engine/Plugins/MovieScene/LevelSequenceEditor/Source/LevelSequenceEditor/Private/LevelSequenceEditorBlueprintLibrary.cpp
+++ b/Engine/Plugins/MovieScene/LevelSequenceEditor/Source/LevelSequenceEditor/Private/LevelSequenceEditorBlueprintLibrary.cpp
@@ -3,13 +3,21 @@
 #include "LevelSequenceEditorBlueprintLibrary.h"
 
 #include "ISequencer.h"
+#include "MVVM/ViewModels/ChannelModel.h"
+#include "MVVM/ViewModels/SectionModel.h"
+#include "MVVM/ViewModels/SequencerEditorViewModel.h"
+#include "MVVM/SectionModelStorageExtension.h"
+#include "MVVM/Selection/Selection.h"
 #include "IKeyArea.h"
 #include "LevelSequence.h"
+#include "Channels/MovieSceneChannel.h"
 
 #include "LevelEditorViewport.h"
 #include "MovieSceneBindingProxy.h"
 #include "Subsystems/AssetEditorSubsystem.h"
 
+#include "ExtensionLibraries/MovieSceneSectionExtensions.h"
+
 #include "MovieSceneObjectBindingID.h"
 #include "MovieSceneSequencePlayer.h"
 
@@ -257,6 +265,64 @@ TArray<FSequencerChannelProxy> ULevelSequenceEditorBlueprintLibrary::GetSelected
 	return OutSelectedChannels;
 }
 
+TArray<FSequencerChannelProxy> ULevelSequenceEditorBlueprintLibrary::GetChannelsWithSelectedKeys()
+{
+	using namespace UE::Sequencer;
+
+	TArray<FSequencerChannelProxy> OutSelectedChannels;
+	TSet<FChannelModel*> ChannelModels;
+	if (CurrentSequencer.IsValid())
+	{
+		const FKeySelection KeySelection = CurrentSequencer.Pin()->GetViewModel()->GetSelection()->KeySelection;
+
+		for (FKeyHandle Key : KeySelection)
+		{
+			TSharedPtr<FChannelModel> Channel = KeySelection.GetModelForKey(Key);
+			if (Channel)
+			{
+				ChannelModels.Add(Channel.Get());
+			}
+		}
+
+		for (FChannelModel* Channel : ChannelModels)
+		{
+			if (Channel)
+			{
+				FSequencerChannelProxy ChannelProxy(Channel->GetChannelName(), Channel->GetSection());
+				OutSelectedChannels.Add(ChannelProxy);
+			}
+		}
+	}
+	return OutSelectedChannels;
+}
+
+TArray<int32> ULevelSequenceEditorBlueprintLibrary::GetSelectedKeys(const FSequencerChannelProxy& ChannelProxy)
+{
+	TArray<int32> SelectedKeys;
+	using namespace UE::Sequencer;
+
+	if (CurrentSequencer.IsValid())
+	{
+		const FKeySelection KeySelection = CurrentSequencer.Pin()->GetViewModel()->GetSelection()->KeySelection;
+
+		for (FKeyHandle Key : KeySelection)
+		{
+			if (TSharedPtr<FChannelModel> Channel = KeySelection.GetModelForKey(Key))
+			{
+				if (Channel->GetChannelName() == ChannelProxy.ChannelName)
+				{
+					int32 Index = Channel->GetChannel()->GetIndex(Key);
+					if (Index != INDEX_NONE)
+					{
+						SelectedKeys.Add(Index);
+					}
+				}
+			}
+		}
+	}
+	return SelectedKeys;
+}
+
 TArray<UMovieSceneFolder*> ULevelSequenceEditorBlueprintLibrary::GetSelectedFolders()
 {
 	TArray<UMovieSceneFolder*> OutSelectedFolders;
@@ -324,6 +390,42 @@ void ULevelSequenceEditorBlueprintLibrary::SelectChannels(const TArray<FSequence
 	}
 }
 
+void ULevelSequenceEditorBlueprintLibrary::SelectKeys(const FSequencerChannelProxy& ChannelProxy, const TArray<int32>& Indices)
+{
+	using namespace UE::Sequencer;
+	if (CurrentSequencer.IsValid())
+	{
+		if (UMovieSceneSection* Section = ChannelProxy.Section)
+		{
+			FSectionModelStorageExtension* SectionModelStorage = CurrentSequencer.Pin()->GetViewModel()->GetRootModel()->CastDynamic<FSectionModelStorageExtension>();
+			check(SectionModelStorage);
+
+			TSharedPtr<FSectionModel> SectionHandle = SectionModelStorage->FindModelForSection(Section);
+			if (SectionHandle)
+			{
+				TParentFirstChildIterator<FChannelGroupModel> KeyAreaNodes = SectionHandle->GetParentTrackModel().AsModel()->GetDescendantsOfType<FChannelGroupModel>();
+				for (const TViewModelPtr<FChannelGroupModel>& KeyAreaNode : KeyAreaNodes)
+				{
+					if (KeyAreaNode->GetChannelName() == ChannelProxy.ChannelName)
+					{
+						if (TSharedPtr<FChannelModel> ChannelModel = KeyAreaNode->GetChannel(Section))
+						{
+							FMovieSceneChannel* MovieSceneChannel = ChannelModel->GetChannel();
+							FKeySelection& KeySelection = CurrentSequencer.Pin()->GetViewModel()->GetSelection()->KeySelection;
+							for (int32 Index : Indices)
+							{
+								FKeyHandle KeyHandle = MovieSceneChannel->GetHandle(Index);
+								KeySelection.Select(ChannelModel, KeyHandle);
+							}
+							break;
+						}
+					}
+				}
+			}
+		}
+	}
+}
+
 void ULevelSequenceEditorBlueprintLibrary::SelectFolders(const TArray<UMovieSceneFolder*>& Folders)
 {
 	if (CurrentSequencer.IsValid())
diff --git a/Engine/Plugins/MovieScene/LevelSequenceEditor/Source/LevelSequenceEditor/Private/LevelSequenceEditorSubsystem.cpp b/Engine/Plugins/MovieScene/LevelSequenceEditor/Source/LevelSequenceEditor/Private/LevelSequenceEditorSubsystem.cpp
--- a/Engine/Plugins/MovieScene/LevelSequenceEditor/Source/LevelSequenceEditor/Private/LevelSequenceEditorSubsystem.cpp
+++ b/Engine/Plugins/MovieScene/LevelSequenceEditor/Source/LevelSequenceEditor/Private/LevelSequenceEditorSubsystem.cpp
@@ -3,7 +3,7 @@
 #include "LevelSequenceEditorSubsystem.h"
 #include "MVVM/ViewModels/SequencerEditorViewModel.h"
 #include "Scripting/SequencerScriptingLayer.h"
-
+#include "SequencerCurveEditorObject.h"
 #include "Evaluation/MovieScenePlayback.h"
 #include "ISequencerModule.h"
 #include "Framework/Commands/UICommandList.h"
@@ -239,6 +239,28 @@ USequencerScriptingLayer* ULevelSequenceEditorSubsystem::GetScriptingLayer()
 	return nullptr;
 }
 
+USequencerCurveEditorObject* ULevelSequenceEditorSubsystem::GetCurveEditor()
+{
+	TObjectPtr<USequencerCurveEditorObject> CurveEditorObject;
+	TSharedPtr<ISequencer> Sequencer = GetActiveSequencer();
+	if (Sequencer)
+	{
+		TObjectPtr<USequencerCurveEditorObject> *ExistingCurveEditorObject = CurveEditorObjects.Find(Sequencer);
+		if (ExistingCurveEditorObject)
+		{
+			CurveEditorObject = *ExistingCurveEditorObject;
+		}
+		else
+		{
+			CurveEditorObject = NewObject<USequencerCurveEditorObject>(this);
+			CurveEditorObject->SetSequencer(Sequencer);
+			CurveEditorObjects.Add(Sequencer, CurveEditorObject);
+			CurveEditorArray.Add(CurveEditorObject);
+		}
+	}
+	return CurveEditorObject;
+}
+
 TArray<FMovieSceneBindingProxy> ULevelSequenceEditorSubsystem::AddActors(const TArray<AActor*>& InActors)
 {
 	TArray<FMovieSceneBindingProxy> BindingProxies;
diff --git a/Engine/Plugins/MovieScene/LevelSequenceEditor/Source/LevelSequenceEditor/Public/LevelSequenceEditorBlueprintLibrary.h b/Engine/Plugins/MovieScene/LevelSequenceEditor/Source/LevelSequenceEditor/Public/LevelSequenceEditorBlueprintLibrary.h
--- a/Engine/Plugins/MovieScene/LevelSequenceEditor/Source/LevelSequenceEditor/Public/LevelSequenceEditorBlueprintLibrary.h
+++ b/Engine/Plugins/MovieScene/LevelSequenceEditor/Source/LevelSequenceEditor/Public/LevelSequenceEditorBlueprintLibrary.h
@@ -3,6 +3,7 @@
 #pragma once
 
 #include "Kismet/BlueprintFunctionLibrary.h"
+#include "SequencerCurveEditorObject.h"
 #include "LevelSequenceEditorBlueprintLibrary.generated.h"
 
 struct FMovieSceneBindingProxy;
@@ -16,27 +17,6 @@ class UMovieSceneSection;
 class UMovieSceneSubSection;
 class UMovieSceneTrack;
 
-USTRUCT(BlueprintType)
-struct FSequencerChannelProxy
-{
-	GENERATED_BODY()
-
-	FSequencerChannelProxy()
-		: Section(nullptr)
-	{}
-
-	FSequencerChannelProxy(const FName& InChannelName, UMovieSceneSection* InSection)
-		: ChannelName(InChannelName)
-		, Section(InSection)
-	{}
-
-	UPROPERTY(BlueprintReadWrite, Category=Channel)
-	FName ChannelName;
-
-	UPROPERTY(BlueprintReadWrite, Category=Channel)
-	TObjectPtr<UMovieSceneSection> Section;
-};
-
 
 UCLASS()
 class LEVELSEQUENCEEDITOR_API ULevelSequenceEditorBlueprintLibrary : public UBlueprintFunctionLibrary
@@ -163,6 +143,14 @@ public:
 	UFUNCTION(BlueprintPure, Category = "Level Sequence Editor")
 	static TArray<FSequencerChannelProxy> GetSelectedChannels();
 
+	/** Gets the channel with selected keys. */
+	UFUNCTION(BlueprintPure, Category = "Level Sequence Editor")
+	static TArray<FSequencerChannelProxy> GetChannelsWithSelectedKeys();
+
+	/** Gets the selected key indices with this channel */
+	UFUNCTION(BlueprintPure, Category = "Level Sequence Editor")
+	static TArray<int32> GetSelectedKeys(const FSequencerChannelProxy& ChannelProxy);
+
 	/** Gets the currently selected folders. */
 	UFUNCTION(BlueprintPure, Category = "Level Sequence Editor")
 	static TArray<UMovieSceneFolder*> GetSelectedFolders();
@@ -183,6 +171,10 @@ public:
 	UFUNCTION(BlueprintCallable, Category = "Level Sequence Editor")
 	static void SelectChannels(const TArray<FSequencerChannelProxy>& Channels);
 
+	/** Select keys from indices */
+	UFUNCTION(BlueprintCallable, Category = "Level Sequence Editor")
+	static void SelectKeys(const FSequencerChannelProxy& Channel, const TArray<int32>& Indices);
+
 	/** Select folders */
 	UFUNCTION(BlueprintCallable, Category = "Level Sequence Editor")
 	static void SelectFolders(const TArray<UMovieSceneFolder*>& Folders);
@@ -256,27 +248,39 @@ public:
 public:
 
 	/** Get if a custom color for specified channel idendified by it's class and identifier exists */
-	UFUNCTION(BlueprintCallable, Category = "Level Sequence Editor")
+	UE_DEPRECATED(5.4, "Use USequencerCurveEditorObject::HasCustomColorForChannel")
+	UFUNCTION(BlueprintCallable, Category = "Level Sequence Editor",
+	meta = (DeprecatedFunction, DeprecationMessage = "Use USequencerCurveEditorObject::HasCustomColorForChannel"))
 	static bool HasCustomColorForChannel(UClass* Class, const FString& Identifier);
 	
 	/** Get custom color for specified channel idendified by it's class and identifier,if none exists will return white*/
-	UFUNCTION(BlueprintCallable, Category = "Level Sequence Editor")
+	UE_DEPRECATED(5.4, "Use USequencerCurveEditorObject::GetCustomColorForChannel")
+	UFUNCTION(BlueprintCallable, Category = "Level Sequence Editor",
+	meta = (DeprecatedFunction, DeprecationMessage = "Use USequencerCurveEditorObject::HasCustomColorForChannel"))
 	static FLinearColor GetCustomColorForChannel(UClass* Class, const FString& Identifier);
 	
 	/** Set Custom Color for specified channel idendified by it's class and identifier. This will be stored in editor user preferences.*/
-	UFUNCTION(BlueprintCallable, Category = "Level Sequence Editor")
+	UE_DEPRECATED(5.4, "Use USequencerCurveEditorObject::SetCustomColorForChannel")
+	UFUNCTION(BlueprintCallable, Category = "Level Sequence Editor",
+	meta = (DeprecatedFunction, DeprecationMessage = "Use USequencerCurveEditorObject::SetCustomColorForChannel"))
 	static void SetCustomColorForChannel(UClass* Class, const FString& Identifier, const FLinearColor& NewColor);
 	
 	/** Set Custom Color for specified channels idendified by it's class and identifiers. This will be stored in editor user preferences.*/
-	UFUNCTION(BlueprintCallable, Category = "Level Sequence Editor")
+	UE_DEPRECATED(5.4, "Use USequencerCurveEditorObject::DeleteColorForChannels")
+	UFUNCTION(BlueprintCallable, Category = "Level Sequence Editor",
+	meta = (DeprecatedFunction, DeprecationMessage = "Use USequencerCurveEditorObject::DeleteColorForChannels"))
 	static void SetCustomColorForChannels(UClass* Class, const TArray<FString>& Identifiers, const TArray<FLinearColor>& NewColors);
 	
 	/** Set Random Colors for specified channels idendified by it's class and identifiers. This will be stored in editor user preferences.*/
-	UFUNCTION(BlueprintCallable, Category = "Level Sequence Editor")
+	UE_DEPRECATED(5.4, "Use USequencerCurveEditorObject::SetRandomColorForChannels")
+	UFUNCTION(BlueprintCallable, Category = "Level Sequence Editor",
+	meta = (DeprecatedFunction, DeprecationMessage = "Use USequencerCurveEditorObject::SetRandomColorForChannels"))
 	static void SetRandomColorForChannels(UClass* Class, const TArray<FString>& Identifiers);
 	
 	/** Delete for specified channel idendified by it's class and identifier.*/
-	UFUNCTION(BlueprintCallable, Category = "Level Sequence Editor")
+	UE_DEPRECATED(5.4, "Use USequencerCurveEditorObject::DeleteColorForChannels")
+	UFUNCTION(BlueprintCallable, Category = "Level Sequence Editor",
+	meta = (DeprecatedFunction, DeprecationMessage = "Use USequencerCurveEditorObject::DeleteColorForChannels"))
 	static void DeleteColorForChannels(UClass* Class, FString& Identifier);
 
 public:
diff --git a/Engine/Plugins/MovieScene/LevelSequenceEditor/Source/LevelSequenceEditor/Public/LevelSequenceEditorSubsystem.h b/Engine/Plugins/MovieScene/LevelSequenceEditor/Source/LevelSequenceEditor/Public/LevelSequenceEditorSubsystem.h
--- a/Engine/Plugins/MovieScene/LevelSequenceEditor/Source/LevelSequenceEditor/Public/LevelSequenceEditorSubsystem.h
+++ b/Engine/Plugins/MovieScene/LevelSequenceEditor/Source/LevelSequenceEditor/Public/LevelSequenceEditorSubsystem.h
@@ -28,6 +28,7 @@ class UMovieSceneCompiledDataManager;
 class UMovieSceneFolder;
 class UMovieSceneSection;
 class USequencerScriptingLayer;
+class USequencerCurveEditorObject;
 
 USTRUCT(BlueprintType)
 struct FMovieSceneScriptingParams
@@ -60,6 +61,10 @@ public:
 	UFUNCTION(BlueprintPure, Category = "Level Sequence Editor")
 	USequencerScriptingLayer* GetScriptingLayer();
 
+	/** Retrieve the curve editor */
+	UFUNCTION(BlueprintPure, Category = "Level Sequence Editor")
+	USequencerCurveEditorObject* GetCurveEditor();
+
 	/** Add existing actors to Sequencer. Tracks will be automatically added based on default track settings. */
 	UFUNCTION(BlueprintCallable, Category = "Level Sequence Editor")
 	TArray<FMovieSceneBindingProxy> AddActors(const TArray<AActor*>& Actors);
@@ -224,6 +229,12 @@ private:
 	/* List of sequencers that have been created */
 	TArray<TWeakPtr<ISequencer>> Sequencers;
 
+	/* Map of curve editors with their sequencers*/
+	TMap<TWeakPtr<ISequencer>, TObjectPtr<USequencerCurveEditorObject>> CurveEditorObjects;
+	/* property array of the curve editors*/
+	UPROPERTY()
+	TArray<TObjectPtr<USequencerCurveEditorObject>> CurveEditorArray;
+
 	TSharedPtr<FUICommandList> CommandList;
 
 	TSharedPtr<FExtender> TransformMenuExtender;
diff --git a/Engine/Plugins/MovieScene/SequencerScripting/Source/SequencerScripting/Private/ExtensionLibraries/MovieSceneSectionExtensions.cpp b/Engine/Plugins/MovieScene/SequencerScripting/Source/SequencerScripting/Private/ExtensionLibraries/MovieSceneSectionExtensions.cpp
--- a/Engine/Plugins/MovieScene/SequencerScripting/Source/SequencerScripting/Private/ExtensionLibraries/MovieSceneSectionExtensions.cpp
+++ b/Engine/Plugins/MovieScene/SequencerScripting/Source/SequencerScripting/Private/ExtensionLibraries/MovieSceneSectionExtensions.cpp
@@ -2,6 +2,7 @@
 
 #include "ExtensionLibraries/MovieSceneSectionExtensions.h"
 #include "Channels/MovieSceneChannelProxy.h"
+#include "Channels/MovieSceneChannelHandle.h"
 #include "Evaluation/MovieSceneSequenceTransform.h"
 #include "KeysAndChannels/MovieSceneScriptingBool.h"
 #include "Evaluation/MovieSceneTimeTransform.h"
@@ -486,6 +487,20 @@ void GetScriptingChannelsForChannel(FMovieSceneChannelProxy& ChannelProxy,TWeakO
 	}
 }
 
+template<typename ChannelType, typename ScriptingChannelType>
+UMovieSceneScriptingChannel* GetScriptingChannelForChannel(const FMovieSceneChannel* MovieSceneChannel, FMovieSceneChannelProxy& ChannelProxy, const FMovieSceneChannelHandle& ChanneHandle, TWeakObjectPtr<UMovieSceneSequence> Sequence, TWeakObjectPtr<UMovieSceneSection> Section, const FName& ChannelName)
+{
+	FName UniqueChannelName = MakeUniqueObjectName(GetTransientPackage(), ScriptingChannelType::StaticClass(), ChannelName);
+	ScriptingChannelType* ScriptingChannel = NewObject<ScriptingChannelType>(GetTransientPackage(), UniqueChannelName);
+	int32 Index = ChannelProxy.FindIndex(ChanneHandle.GetChannelTypeName(), MovieSceneChannel);
+	SetScriptingChannelHandle<ChannelType, ScriptingChannelType>(ScriptingChannel, ChannelProxy, Index);
+	ScriptingChannel->OwningSection = Section;
+	ScriptingChannel->OwningSequence = Sequence;
+	ScriptingChannel->ChannelName = ChannelName;
+	return ScriptingChannel;
+}
+
+
 TArray<UMovieSceneScriptingChannel*> UMovieSceneSectionExtensions::GetAllChannels(UMovieSceneSection* Section)
 {
 	TArray<UMovieSceneScriptingChannel*> Channels;
@@ -541,10 +556,69 @@ TArray<UMovieSceneScriptingChannel*> UMovieSceneSectionExtensions::GetChannelsBy
 		FFrame::KismetExecutionMessage(TEXT("Unsupported ChannelType for FindChannelsByType!"), ELogVerbosity::Error);
 	}
 
-
 	return Channels;
 }
 
+UMovieSceneScriptingChannel* UMovieSceneSectionExtensions::GetChannel(UMovieSceneSection* Section, const FName& ChannelName)
+{
+#if WITH_EDITORONLY_DATA
+
+	if (!Section)
+	{
+		FFrame::KismetExecutionMessage(TEXT("Cannot get channel for null section"), ELogVerbosity::Error);
+		return nullptr;
+	}
+	FMovieSceneChannel* MovieSceneChannel =  UMovieSceneSectionExtensions::GetMovieSceneChannel(Section, ChannelName);
+	if (!MovieSceneChannel)
+	{
+		FFrame::KismetExecutionMessage(TEXT("Cannot get channel with this name"), ELogVerbosity::Error);
+		return nullptr;
+	}
+
+	FMovieSceneChannelProxy& ChannelProxy = Section->GetChannelProxy();
+	FMovieSceneChannelHandle ChanneHandle = ChannelProxy.GetChannelByName(ChannelName);
+	TWeakObjectPtr<UMovieSceneSequence> Sequence = Section->GetTypedOuter<UMovieSceneSequence>();
+	if (ChanneHandle.GetChannelTypeName() == FMovieSceneBoolChannel::StaticStruct()->GetFName())
+	{
+		return GetScriptingChannelForChannel<FMovieSceneBoolChannel, UMovieSceneScriptingBoolChannel>(MovieSceneChannel, ChannelProxy, ChanneHandle, Sequence, Section, ChannelName);
+	} 
+	else if (ChanneHandle.GetChannelTypeName() == FMovieSceneByteChannel::StaticStruct()->GetFName())
+	{
+		return GetScriptingChannelForChannel<FMovieSceneByteChannel, UMovieSceneScriptingByteChannel>(MovieSceneChannel, ChannelProxy, ChanneHandle, Sequence, Section, ChannelName);
+	}
+	else if (ChanneHandle.GetChannelTypeName() == FMovieSceneFloatChannel::StaticStruct()->GetFName())
+	{
+		return GetScriptingChannelForChannel<FMovieSceneFloatChannel, UMovieSceneScriptingFloatChannel>(MovieSceneChannel, ChannelProxy, ChanneHandle, Sequence, Section, ChannelName);
+	}
+	else if (ChanneHandle.GetChannelTypeName() == FMovieSceneDoubleChannel::StaticStruct()->GetFName())
+	{
+		return GetScriptingChannelForChannel<FMovieSceneDoubleChannel, UMovieSceneScriptingDoubleChannel>(MovieSceneChannel, ChannelProxy, ChanneHandle, Sequence, Section, ChannelName);
+	}
+	else if (ChanneHandle.GetChannelTypeName() == FMovieSceneIntegerChannel::StaticStruct()->GetFName())
+	{
+		return GetScriptingChannelForChannel<FMovieSceneIntegerChannel, UMovieSceneScriptingIntegerChannel>(MovieSceneChannel, ChannelProxy, ChanneHandle, Sequence, Section, ChannelName);
+	}
+	else if (ChanneHandle.GetChannelTypeName() == FMovieSceneStringChannel::StaticStruct()->GetFName())
+	{
+		return GetScriptingChannelForChannel<FMovieSceneStringChannel, UMovieSceneScriptingStringChannel>(MovieSceneChannel, ChannelProxy, ChanneHandle, Sequence, Section, ChannelName);
+	}
+	else if (ChanneHandle.GetChannelTypeName() == FMovieSceneEventChannel::StaticStruct()->GetFName())
+	{
+		return GetScriptingChannelForChannel<FMovieSceneEventChannel, UMovieSceneScriptingEventChannel>(MovieSceneChannel, ChannelProxy, ChanneHandle, Sequence, Section, ChannelName);
+	}
+	else if (ChanneHandle.GetChannelTypeName() == FMovieSceneActorReferenceData::StaticStruct()->GetFName())
+	{
+		return GetScriptingChannelForChannel<FMovieSceneActorReferenceData, UMovieSceneScriptingActorReferenceChannel>(MovieSceneChannel, ChannelProxy, ChanneHandle, Sequence, Section, ChannelName);
+	}
+	else if (ChanneHandle.GetChannelTypeName() == FMovieSceneObjectPathChannel::StaticStruct()->GetFName())
+	{
+		return GetScriptingChannelForChannel<FMovieSceneObjectPathChannel, UMovieSceneScriptingObjectPathChannel>(MovieSceneChannel, ChannelProxy, ChanneHandle, Sequence, Section, ChannelName);
+	}
+#endif
+	FFrame::KismetExecutionMessage(TEXT("Could not find channel"), ELogVerbosity::Error);
+	return nullptr;
+}
+
 bool GetSubSectionChain(UMovieSceneSubSection* InSubSection, UMovieSceneSequence* ParentSequence, TArray<UMovieSceneSubSection*>& SubSectionChain)
 {
 	UMovieScene* ParentMovieScene = ParentSequence->GetMovieScene();
@@ -598,3 +672,27 @@ int32 UMovieSceneSectionExtensions::GetParentSequenceFrame(UMovieSceneSubSection
 	return LocalFrameTime.GetFrame().Value;
 }
 
+FMovieSceneChannel* UMovieSceneSectionExtensions::GetMovieSceneChannel(UMovieSceneSection* Section, const FName& Name)
+{
+#if WITH_EDITORONLY_DATA
+
+	FMovieSceneChannelProxy& ChannelProxy = Section->GetChannelProxy();
+	for (const FMovieSceneChannelEntry& Entry : Section->GetChannelProxy().GetAllEntries())
+	{
+		TArrayView<FMovieSceneChannel* const>  Channels = Entry.GetChannels();
+		TArrayView<const FMovieSceneChannelMetaData> MetaData = Entry.GetMetaData();
+		for (int32 Index = 0; Index < MetaData.Num(); ++Index)
+		{
+			if (MetaData[Index].bEnabled)
+			{
+				const FName ChannelName = MetaData[Index].Name;
+				if (ChannelName == Name)
+				{
+					return Channels[Index];
+				}
+			}
+		}
+	}
+#endif
+	return nullptr;
+}
diff --git a/Engine/Plugins/MovieScene/SequencerScripting/Source/SequencerScripting/Private/KeysAndChannels/MovieSceneScriptingActorReference.h b/Engine/Plugins/MovieScene/SequencerScripting/Source/SequencerScripting/Private/KeysAndChannels/MovieSceneScriptingActorReference.h
--- a/Engine/Plugins/MovieScene/SequencerScripting/Source/SequencerScripting/Private/KeysAndChannels/MovieSceneScriptingActorReference.h
+++ b/Engine/Plugins/MovieScene/SequencerScripting/Source/SequencerScripting/Private/KeysAndChannels/MovieSceneScriptingActorReference.h
@@ -104,6 +104,18 @@ public:
 		return Impl::GetKeysInChannel(ChannelHandle, OwningSequence, OwningSection);
 	}
 
+	/**
+	* Gets the keys in this channel specified by the specific index
+	* @Indices  The indices from which to get the keys from
+	* @return	An array of UMovieSceneScriptingKey's contained by this channel.
+	*			Returns all keys specified by the indices, even if out of range.
+	*/
+	UFUNCTION(BlueprintCallable, Category = "Sequencer|Keys", meta = (DisplayName = "Get Keys By Index (Actor Reference)"))
+	virtual TArray<UMovieSceneScriptingKey*> GetKeysByIndex(const TArray<int32>& Indices) const override
+	{
+		return Impl::GetKeysInChannelByIndex(ChannelHandle, OwningSequence, OwningSection, Indices);
+	}
+
 	/**
 	* Set this channel's default value that should be used when no keys are present.
 	* Sets bHasDefaultValue to true automatically.
diff --git a/Engine/Plugins/MovieScene/SequencerScripting/Source/SequencerScripting/Private/KeysAndChannels/MovieSceneScriptingBool.h b/Engine/Plugins/MovieScene/SequencerScripting/Source/SequencerScripting/Private/KeysAndChannels/MovieSceneScriptingBool.h
--- a/Engine/Plugins/MovieScene/SequencerScripting/Source/SequencerScripting/Private/KeysAndChannels/MovieSceneScriptingBool.h
+++ b/Engine/Plugins/MovieScene/SequencerScripting/Source/SequencerScripting/Private/KeysAndChannels/MovieSceneScriptingBool.h
@@ -100,6 +100,18 @@ public:
 		return Impl::GetKeysInChannel(ChannelHandle, OwningSequence, OwningSection);
 	}
 
+	/**
+	* Gets the keys in this channel specified by the specific index
+	* @Indices  The indices from which to get the keys from
+	* @return	An array of UMovieSceneScriptingKey's contained by this channel.
+	*			Returns all keys specified by the indices, even if out of range.
+	*/
+	UFUNCTION(BlueprintCallable, Category = "Sequencer|Keys", meta = (DisplayName = "Get Keys By Index (Bool)"))
+	virtual TArray<UMovieSceneScriptingKey*> GetKeysByIndex(const TArray<int32>& Indices) const override
+	{
+		return Impl::GetKeysInChannelByIndex(ChannelHandle, OwningSequence, OwningSection, Indices);
+	}
+
 	/**
 	* Returns number of keys in this channel.
 	*/
diff --git a/Engine/Plugins/MovieScene/SequencerScripting/Source/SequencerScripting/Private/KeysAndChannels/MovieSceneScriptingByte.h b/Engine/Plugins/MovieScene/SequencerScripting/Source/SequencerScripting/Private/KeysAndChannels/MovieSceneScriptingByte.h
--- a/Engine/Plugins/MovieScene/SequencerScripting/Source/SequencerScripting/Private/KeysAndChannels/MovieSceneScriptingByte.h
+++ b/Engine/Plugins/MovieScene/SequencerScripting/Source/SequencerScripting/Private/KeysAndChannels/MovieSceneScriptingByte.h
@@ -103,6 +103,18 @@ public:
 		return Impl::GetKeysInChannel(ChannelHandle, OwningSequence, OwningSection);
 	}
 
+	/**
+	* Gets the keys in this channel specified by the specific index
+	* @Indices  The indices from which to get the keys from
+	* @return	An array of UMovieSceneScriptingKey's contained by this channel.
+	*			Returns all keys specified by the indices, even if out of range.
+	*/
+	UFUNCTION(BlueprintCallable, Category = "Sequencer|Keys", meta = (DisplayName = "Get Keys By Index (Enum)"))
+	virtual TArray<UMovieSceneScriptingKey*> GetKeysByIndex(const TArray<int32>& Indices) const override
+	{
+		return Impl::GetKeysInChannelByIndex(ChannelHandle, OwningSequence, OwningSection, Indices);
+	}
+
 	/**
 	* Set this channel's default value that should be used when no keys are present.
 	* Sets bHasDefaultValue to true automatically.
diff --git a/Engine/Plugins/MovieScene/SequencerScripting/Source/SequencerScripting/Private/KeysAndChannels/MovieSceneScriptingChannel.h b/Engine/Plugins/MovieScene/SequencerScripting/Source/SequencerScripting/Private/KeysAndChannels/MovieSceneScriptingChannel.h
--- a/Engine/Plugins/MovieScene/SequencerScripting/Source/SequencerScripting/Private/KeysAndChannels/MovieSceneScriptingChannel.h
+++ b/Engine/Plugins/MovieScene/SequencerScripting/Source/SequencerScripting/Private/KeysAndChannels/MovieSceneScriptingChannel.h
@@ -44,6 +44,15 @@ public:
 	*/
 	virtual TArray<UMovieSceneScriptingKey*> GetKeys() const PURE_VIRTUAL(UMovieSceneScriptingChannel::GetKeys, return TArray<UMovieSceneScriptingKey*>(););
 
+	/**
+	* Gets the keys in this channel specified by the specific index
+	* @Indices  The indices from which to get the keys from
+	* @return	An array of UMovieSceneScriptingKey's contained by this channel.
+	*			Returns all keys specified by the indices, even if out of range.
+	*/
+	virtual TArray<UMovieSceneScriptingKey*> GetKeysByIndex(const TArray<int32> &Indices) const PURE_VIRTUAL(UMovieSceneScriptingChannel::GetKeysByIndex, return TArray<UMovieSceneScriptingKey*>(););
+
+
 	UPROPERTY(BlueprintReadOnly, Category="Sequencer|Keys")
 	FName ChannelName;
 };
@@ -122,6 +131,42 @@ struct TMovieSceneScriptingChannel
 		FFrame::KismetExecutionMessage(TEXT("Invalid ChannelHandle for MovieSceneScriptingChannel, failed to remove key."), ELogVerbosity::Error);
 	}
 
+	static TArray<UMovieSceneScriptingKey*> GetKeysInChannelByIndex(TMovieSceneChannelHandle<ChannelType> ChannelHandle, TWeakObjectPtr<UMovieSceneSequence> Sequence, TWeakObjectPtr<UMovieSceneSection> Section,
+		const TArray<int32>& Indices)
+	{
+		TArray<UMovieSceneScriptingKey*> OutScriptingKeys;
+		ChannelType* Channel = ChannelHandle.Get();
+		if (Channel)
+		{
+			TArray<FFrameNumber> OutTimes;
+			TArray<FKeyHandle> OutKeys;
+			Channel->GetKeys(TRange<FFrameNumber>(), &OutTimes, &OutKeys);
+
+			for (int32 Index: Indices)
+			{
+				if (Index >= 0 && Index < OutTimes.Num())
+				{
+					ScriptingKeyType * Key = NewObject<ScriptingKeyType>();
+					Key->KeyHandle = OutKeys[Index];
+					Key->ChannelHandle = ChannelHandle;
+					Key->OwningSequence = Sequence;
+					Key->OwningSection = Section;
+					OutScriptingKeys.Add(Key);
+				}
+				else
+				{
+					FFrame::KismetExecutionMessage(TEXT("Invalid index for MovieSceneScriptingChannel, failed to get keys by index."), ELogVerbosity::Error);
+				}
+			}
+		}
+		else
+		{
+			FFrame::KismetExecutionMessage(TEXT("Invalid ChannelHandle for MovieSceneScriptingChannel, failed to get keys."), ELogVerbosity::Error);
+		}
+
+		return OutScriptingKeys;
+	}
+
 	static TArray<UMovieSceneScriptingKey*> GetKeysInChannel(TMovieSceneChannelHandle<ChannelType> ChannelHandle, TWeakObjectPtr<UMovieSceneSequence> Sequence, TWeakObjectPtr<UMovieSceneSection> Section)
 	{
 		TArray<UMovieSceneScriptingKey*> OutScriptingKeys;
diff --git a/Engine/Plugins/MovieScene/SequencerScripting/Source/SequencerScripting/Private/KeysAndChannels/MovieSceneScriptingDouble.h b/Engine/Plugins/MovieScene/SequencerScripting/Source/SequencerScripting/Private/KeysAndChannels/MovieSceneScriptingDouble.h
--- a/Engine/Plugins/MovieScene/SequencerScripting/Source/SequencerScripting/Private/KeysAndChannels/MovieSceneScriptingDouble.h
+++ b/Engine/Plugins/MovieScene/SequencerScripting/Source/SequencerScripting/Private/KeysAndChannels/MovieSceneScriptingDouble.h
@@ -266,6 +266,17 @@ public:
 		return Impl::GetKeysInChannel(ChannelHandle, OwningSequence, OwningSection);
 	}
 
+	/**
+	* Gets the keys in this channel specified by the specific index
+	* @Indices  The indices from which to get the keys from
+	* @return	An array of UMovieSceneScriptingKey's contained by this channel.
+	*			Returns all keys specified by the indices, even if out of range.
+	*/
+	UFUNCTION(BlueprintCallable, Category = "Sequencer|Keys", meta = (DisplayName = "Get Keys By Index (Double)"))
+	virtual TArray<UMovieSceneScriptingKey*> GetKeysByIndex(const TArray<int32>& Indices) const override
+	{
+		return Impl::GetKeysInChannelByIndex(ChannelHandle, OwningSequence, OwningSection, Indices);
+	}
 	/**
 	* Returns number of keys in this channel.
 	*/
diff --git a/Engine/Plugins/MovieScene/SequencerScripting/Source/SequencerScripting/Private/KeysAndChannels/MovieSceneScriptingEvent.h b/Engine/Plugins/MovieScene/SequencerScripting/Source/SequencerScripting/Private/KeysAndChannels/MovieSceneScriptingEvent.h
--- a/Engine/Plugins/MovieScene/SequencerScripting/Source/SequencerScripting/Private/KeysAndChannels/MovieSceneScriptingEvent.h
+++ b/Engine/Plugins/MovieScene/SequencerScripting/Source/SequencerScripting/Private/KeysAndChannels/MovieSceneScriptingEvent.h
@@ -97,6 +97,18 @@ public:
 		return Impl::GetKeysInChannel(ChannelHandle, OwningSequence, OwningSection);
 	}
 
+	/**
+	* Gets the keys in this channel specified by the specific index
+	* @Indices  The indices from which to get the keys from
+	* @return	An array of UMovieSceneScriptingKey's contained by this channel.
+	*			Returns all keys specified by the indices, even if out of range.
+	*/
+	UFUNCTION(BlueprintCallable, Category = "Sequencer|Keys", meta = (DisplayName = "Get Keys By Index (Event)"))
+	virtual TArray<UMovieSceneScriptingKey*> GetKeysByIndex(const TArray<int32>& Indices) const override
+	{
+		return Impl::GetKeysInChannelByIndex(ChannelHandle, OwningSequence, OwningSection, Indices);
+	}
+
 public:
 	TWeakObjectPtr<UMovieSceneSequence> OwningSequence;
 	TMovieSceneChannelHandle<FMovieSceneEventChannel> ChannelHandle;
diff --git a/Engine/Plugins/MovieScene/SequencerScripting/Source/SequencerScripting/Private/KeysAndChannels/MovieSceneScriptingFloat.h b/Engine/Plugins/MovieScene/SequencerScripting/Source/SequencerScripting/Private/KeysAndChannels/MovieSceneScriptingFloat.h
--- a/Engine/Plugins/MovieScene/SequencerScripting/Source/SequencerScripting/Private/KeysAndChannels/MovieSceneScriptingFloat.h
+++ b/Engine/Plugins/MovieScene/SequencerScripting/Source/SequencerScripting/Private/KeysAndChannels/MovieSceneScriptingFloat.h
@@ -419,6 +419,25 @@ public:
 		}
 	}
 
+	/**
+	* Gets the keys in this channel specified by the specific index
+	* @Indices  The indices from which to get the keys from
+	* @return	An array of UMovieSceneScriptingKey's contained by this channel.
+	*			Returns all keys specified by the indices, even if out of range.
+	*/
+	UFUNCTION(BlueprintCallable, Category = "Sequencer|Keys", meta = (DisplayName = "Get Keys By Index (Float)"))
+	virtual TArray<UMovieSceneScriptingKey*> GetKeysByIndex(const TArray<int32>& Indices) const override
+	{
+		if (FloatChannelHandle.Get())
+		{
+			return FloatImpl::GetKeysInChannelByIndex(FloatChannelHandle, OwningSequence, OwningSection, Indices);
+		}
+		else
+		{
+			return DoubleImpl::GetKeysInChannelByIndex(DoubleChannelHandle, OwningSequence, OwningSection, Indices);
+		}
+	}
+
 	/**
 	* Returns number of keys in this channel.
 	*/
diff --git a/Engine/Plugins/MovieScene/SequencerScripting/Source/SequencerScripting/Private/KeysAndChannels/MovieSceneScriptingInteger.h b/Engine/Plugins/MovieScene/SequencerScripting/Source/SequencerScripting/Private/KeysAndChannels/MovieSceneScriptingInteger.h
--- a/Engine/Plugins/MovieScene/SequencerScripting/Source/SequencerScripting/Private/KeysAndChannels/MovieSceneScriptingInteger.h
+++ b/Engine/Plugins/MovieScene/SequencerScripting/Source/SequencerScripting/Private/KeysAndChannels/MovieSceneScriptingInteger.h
@@ -102,6 +102,18 @@ public:
 		return Impl::GetKeysInChannel(ChannelHandle, OwningSequence, OwningSection);
 	}
 
+	/**
+	* Gets the keys in this channel specified by the specific index
+	* @Indices  The indices from which to get the keys from
+	* @return	An array of UMovieSceneScriptingKey's contained by this channel.
+	*			Returns all keys specified by the indices, even if out of range.
+	*/
+	UFUNCTION(BlueprintCallable, Category = "Sequencer|Keys", meta = (DisplayName = "Get Keys By Index (Integer)"))
+	virtual TArray<UMovieSceneScriptingKey*> GetKeysByIndex(const TArray<int32>& Indices) const override
+	{
+		return Impl::GetKeysInChannelByIndex(ChannelHandle, OwningSequence, OwningSection, Indices);
+	}
+
 	/**
 	* Set this channel's default value that should be used when no keys are present.
 	* Sets bHasDefaultValue to true automatically.
diff --git a/Engine/Plugins/MovieScene/SequencerScripting/Source/SequencerScripting/Private/KeysAndChannels/MovieSceneScriptingObjectPath.h b/Engine/Plugins/MovieScene/SequencerScripting/Source/SequencerScripting/Private/KeysAndChannels/MovieSceneScriptingObjectPath.h
--- a/Engine/Plugins/MovieScene/SequencerScripting/Source/SequencerScripting/Private/KeysAndChannels/MovieSceneScriptingObjectPath.h
+++ b/Engine/Plugins/MovieScene/SequencerScripting/Source/SequencerScripting/Private/KeysAndChannels/MovieSceneScriptingObjectPath.h
@@ -105,6 +105,18 @@ public:
 		return Impl::GetKeysInChannel(ChannelHandle, OwningSequence, OwningSection);
 	}
 
+	/**
+	* Gets the keys in this channel specified by the specific index
+	* @Indices  The indices from which to get the keys from
+	* @return	An array of UMovieSceneScriptingKey's contained by this channel.
+	*			Returns all keys specified by the indices, even if out of range.
+	*/
+	UFUNCTION(BlueprintCallable, Category = "Sequencer|Keys", meta = (DisplayName = "Get Keys By Index (Object Path)"))
+	virtual TArray<UMovieSceneScriptingKey*> GetKeysByIndex(const TArray<int32>& Indices) const override
+	{
+		return Impl::GetKeysInChannelByIndex(ChannelHandle, OwningSequence, OwningSection, Indices);
+	}
+
 	/**
 	 * Set this channel's default value that should be used when no keys are present.
 	 * Sets bHasDefaultValue to true automatically.
diff --git a/Engine/Plugins/MovieScene/SequencerScripting/Source/SequencerScripting/Private/KeysAndChannels/MovieSceneScriptingString.h b/Engine/Plugins/MovieScene/SequencerScripting/Source/SequencerScripting/Private/KeysAndChannels/MovieSceneScriptingString.h
--- a/Engine/Plugins/MovieScene/SequencerScripting/Source/SequencerScripting/Private/KeysAndChannels/MovieSceneScriptingString.h
+++ b/Engine/Plugins/MovieScene/SequencerScripting/Source/SequencerScripting/Private/KeysAndChannels/MovieSceneScriptingString.h
@@ -102,6 +102,18 @@ public:
 		return Impl::GetKeysInChannel(ChannelHandle, OwningSequence, OwningSection);
 	}
 
+	/**
+	* Gets the keys in this channel specified by the specific index
+	* @Indices  The indices from which to get the keys from
+	* @return	An array of UMovieSceneScriptingKey's contained by this channel.
+	*			Returns all keys specified by the indices, even if out of range.
+	*/
+	UFUNCTION(BlueprintCallable, Category = "Sequencer|Keys", meta = (DisplayName = "Get Keys By Index (String)"))
+	virtual TArray<UMovieSceneScriptingKey*> GetKeysByIndex(const TArray<int32>& Indices) const override
+	{
+		return Impl::GetKeysInChannelByIndex(ChannelHandle, OwningSequence, OwningSection, Indices);
+	}
+
 	/**
 	* Set this channel's default value that should be used when no keys are present.
 	* Sets bHasDefaultValue to true automatically.
diff --git a/Engine/Plugins/MovieScene/SequencerScripting/Source/SequencerScripting/Public/ExtensionLibraries/MovieSceneSectionExtensions.h b/Engine/Plugins/MovieScene/SequencerScripting/Source/SequencerScripting/Public/ExtensionLibraries/MovieSceneSectionExtensions.h
--- a/Engine/Plugins/MovieScene/SequencerScripting/Source/SequencerScripting/Public/ExtensionLibraries/MovieSceneSectionExtensions.h
+++ b/Engine/Plugins/MovieScene/SequencerScripting/Source/SequencerScripting/Public/ExtensionLibraries/MovieSceneSectionExtensions.h
@@ -9,7 +9,7 @@ class UMovieSceneScriptingChannel;
 class UMovieSceneSequence;
 class UMovieSceneSubSection;
 struct FSequencerScriptingRange;
-
+struct FMovieSceneChannel;
 class UMovieSceneSection;
 
 /**
@@ -225,6 +225,19 @@ public:
 	UFUNCTION(BlueprintCallable, Category = "Sequencer|Section", meta = (ScriptMethod, DeterminesOutputType="TrackType", DevelopmentOnly))
 	static TArray<UMovieSceneScriptingChannel*> GetChannelsByType(UMovieSceneSection* Section, TSubclassOf<UMovieSceneScriptingChannel> ChannelType);
 
+
+
+	/**
+	* Get channel from specified section and channel name
+	*
+	* @param Section        The section to use.
+	* @param ChannelName	The name of the channel.
+	* @return The channel if it exists 
+	*/
+	UFUNCTION(BlueprintCallable, Category = "Sequencer|Section", meta = (ScriptMethod, DeterminesOutputType = "TrackType", DevelopmentOnly))
+	static UMovieSceneScriptingChannel* GetChannel(UMovieSceneSection* Section, const FName& ChannelName);
+
+
 	/**
 	 * Get the frame in the space of its parent sequence
 	 *
@@ -235,6 +248,14 @@ public:
 	 */
 	UFUNCTION(BlueprintCallable, Category = "Sequencer|Section", meta = (ScriptMethod))
 	static int32 GetParentSequenceFrame(UMovieSceneSubSection* Section, int32 InFrame, UMovieSceneSequence* ParentSequence);
+
+public:
+	/**
+	* Utility function to get channel from a section and a name
+	*/
+	static SEQUENCERSCRIPTING_API FMovieSceneChannel* GetMovieSceneChannel(UMovieSceneSection* InSection, const FName& InName);
+
+
 };
 
 #if UE_ENABLE_INCLUDE_ORDER_DEPRECATED_IN_5_2
diff --git a/Engine/Plugins/MovieScene/SequencerScripting/Source/SequencerScriptingEditor/Private/SequencerCurveEditorObject.cpp b/Engine/Plugins/MovieScene/SequencerScripting/Source/SequencerScriptingEditor/Private/SequencerCurveEditorObject.cpp
new file mode 100644
--- /dev/null
+++ b/Engine/Plugins/MovieScene/SequencerScripting/Source/SequencerScriptingEditor/Private/SequencerCurveEditorObject.cpp
@@ -0,0 +1,270 @@
+// Copyright Epic Games, Inc. All Rights Reserved.
+#include "SequencerCurveEditorObject.h"
+#include "MVVM/CurveEditorExtension.h"
+#include "MVVM/ViewModels/SequencerEditorViewModel.h"
+#include "ISequencer.h"
+#include "CurveEditor.h"
+#include "CurveEditorTypes.h"
+#include "CurveModel.h"
+#include "Channels/MovieSceneChannel.h"
+#include "ExtensionLibraries/MovieSceneSectionExtensions.h"
+//For custom colors on channels, stored in editor pref's
+#include "CurveEditorSettings.h"
+
+
+TSharedPtr<FCurveEditor> USequencerCurveEditorObject::GetCurveEditor()
+{
+	using namespace UE::Sequencer;
+	if (CurrentSequencer.IsValid())
+	{
+		const TSharedPtr<FSequencerEditorViewModel> SequencerViewModel = CurrentSequencer.Pin()->GetViewModel();
+		if (FCurveEditorExtension* CurveEditorExtension = SequencerViewModel->CastDynamic<FCurveEditorExtension>())
+		{
+			return (CurveEditorExtension->GetCurveEditor());
+		}
+	}
+	return nullptr;
+}
+
+void USequencerCurveEditorObject::OpenCurveEditor()
+{
+	using namespace UE::Sequencer;
+	if (CurrentSequencer.IsValid())
+	{
+		const TSharedPtr<FSequencerEditorViewModel> SequencerViewModel = CurrentSequencer.Pin()->GetViewModel();
+		if (FCurveEditorExtension* CurveEditorExtension = SequencerViewModel->CastDynamic<FCurveEditorExtension>())
+		{
+			CurveEditorExtension->OpenCurveEditor();
+		}
+	}
+}
+
+bool USequencerCurveEditorObject::IsCurveEditorOpen()
+{
+	using namespace UE::Sequencer;
+	if (CurrentSequencer.IsValid())
+	{
+		const TSharedPtr<FSequencerEditorViewModel> SequencerViewModel = CurrentSequencer.Pin()->GetViewModel();
+		if (const FCurveEditorExtension* CurveEditorExtension = SequencerViewModel->CastDynamic<FCurveEditorExtension>())
+		{
+			return CurveEditorExtension->IsCurveEditorOpen();
+		}
+	}
+	return false;
+}
+
+void USequencerCurveEditorObject::CloseCurveEditor()
+{
+	using namespace UE::Sequencer;
+	if (CurrentSequencer.IsValid())
+	{
+		const TSharedPtr<FSequencerEditorViewModel> SequencerViewModel = CurrentSequencer.Pin()->GetViewModel();
+		if (FCurveEditorExtension* CurveEditorExtension = SequencerViewModel->CastDynamic<FCurveEditorExtension>())
+		{
+			CurveEditorExtension->CloseCurveEditor();
+		}
+	}
+}
+
+TArray<FSequencerChannelProxy> USequencerCurveEditorObject::GetChannelsWithSelectedKeys()
+{
+	TArray<FSequencerChannelProxy> OutSelectedChannels;
+	TSharedPtr<FCurveEditor> CurveEditor = GetCurveEditor();
+	if (CurrentSequencer.IsValid() && CurveEditor.IsValid())
+	{
+		const TMap<FCurveModelID, FKeyHandleSet>& SelectionKeyMap = CurveEditor->Selection.GetAll();
+
+		for (const TPair<FCurveModelID, FKeyHandleSet>& Pair : SelectionKeyMap)
+		{
+			if (const FCurveModel* Curve = CurveEditor->FindCurve(Pair.Key))
+			{
+				if (UMovieSceneSection* Section = Cast<UMovieSceneSection>(Curve->GetOwningObject()))
+				{
+					FString IntentionName = Curve->GetIntentionName();
+					FText DisplayName = Curve->GetShortDisplayName();
+					FText LongName = Curve->GetLongDisplayName();
+					FName ChannelName = FName(*DisplayName.ToString());
+					FSequencerChannelProxy ChannelProxy(ChannelName,Section);
+					OutSelectedChannels.Add(ChannelProxy);
+
+				}
+			}
+		}
+	}
+	return OutSelectedChannels;
+}
+
+TArray<int32> USequencerCurveEditorObject::GetSelectedKeys(const FSequencerChannelProxy& ChannelProxy)
+{
+	TArray<int32> SelectedKeys;
+	TSharedPtr<FCurveEditor> CurveEditor = GetCurveEditor();
+	if (CurrentSequencer.IsValid() && CurveEditor.IsValid())
+	{
+		const TMap<FCurveModelID, FKeyHandleSet>& SelectionKeyMap = CurveEditor->Selection.GetAll();
+
+		for (const TPair<FCurveModelID, FKeyHandleSet>& Pair : SelectionKeyMap)
+		{
+			if (const FCurveModel* Curve = CurveEditor->FindCurve(Pair.Key))
+			{
+				if (UMovieSceneSection* Section = Cast<UMovieSceneSection>(Curve->GetOwningObject()))
+				{
+					if (Section == ChannelProxy.Section)
+					{
+						if (FMovieSceneChannel* MovieSceneChannel = UMovieSceneSectionExtensions::GetMovieSceneChannel(Section, ChannelProxy.ChannelName))
+						{
+							TArrayView<const FKeyHandle> HandleArray = Pair.Value.AsArray();
+							for (FKeyHandle Key : HandleArray)
+							{
+								int32 Index = MovieSceneChannel->GetIndex(Key);
+								if (Index != INDEX_NONE)
+								{
+									SelectedKeys.Add(Index);
+								}
+							}
+						}
+					}
+				}
+			}
+		}
+	}
+	return SelectedKeys;
+}
+
+void USequencerCurveEditorObject::EmptySelection()
+{
+	TSharedPtr<FCurveEditor> CurveEditor = GetCurveEditor();
+	if (CurrentSequencer.IsValid() && CurveEditor.IsValid())
+	{
+		CurveEditor->Selection.Clear();
+	}
+}
+
+FCurveModelID USequencerCurveEditorObject::GetCurve(UMovieSceneSection* InSection, const FName& InName)
+{
+	TSharedPtr<FCurveEditor> CurveEditor = GetCurveEditor();
+	if (CurrentSequencer.IsValid() && CurveEditor.IsValid())
+	{
+		const TMap<FCurveModelID, TUniquePtr<FCurveModel>>& Curves = CurveEditor->GetCurves();
+		for (const TPair <FCurveModelID, TUniquePtr<FCurveModel>>& Pair : Curves)
+		{
+			if (Pair.Value.IsValid() && Pair.Value->GetOwningObject() == InSection && Pair.Value->GetLongIntentionName() == InName)
+			{
+				return Pair.Key;
+			}
+		}
+	}
+	return FCurveModelID::Unique();
+}
+
+void USequencerCurveEditorObject::SelectKeys(const FSequencerChannelProxy& ChannelProxy, const TArray<int32>& Indices)
+{
+	TSharedPtr<FCurveEditor> CurveEditor = GetCurveEditor();
+	if (CurrentSequencer.IsValid() && CurveEditor.IsValid())
+	{
+		if (UMovieSceneSection* Section = ChannelProxy.Section)
+		{
+			FCurveModelID CurveModelID = USequencerCurveEditorObject::GetCurve(Section, ChannelProxy.ChannelName);
+			{
+				if (FMovieSceneChannel* MovieSceneChannel = UMovieSceneSectionExtensions::GetMovieSceneChannel(Section, ChannelProxy.ChannelName))
+				{
+					TArray<FKeyHandle> Handles;
+					for (int32 Index : Indices)
+					{
+						FKeyHandle Handle = MovieSceneChannel->GetHandle(Index);
+						if (Handle != FKeyHandle::Invalid())
+						{
+							Handles.Add(Handle);
+						}
+					}
+					CurveEditor->Selection.Add(CurveModelID, ECurvePointType::Key, Handles);
+				}
+			}
+		}
+	}
+}
+
+void USequencerCurveEditorObject::SetSequencer(TSharedPtr<ISequencer>& InSequencer)
+{
+	using namespace UE::Sequencer;
+	CurrentSequencer = TWeakPtr<ISequencer>(InSequencer);
+}
+
+bool USequencerCurveEditorObject::HasCustomColorForChannel(UClass* Class, const FString& Identifier)
+{
+	const UCurveEditorSettings* Settings = GetDefault<UCurveEditorSettings>();
+	if (Settings)
+	{
+		TOptional<FLinearColor> OptColor = Settings->GetCustomColor(Class, Identifier);
+		return OptColor.IsSet();
+	}
+	return false;
+}
+
+FLinearColor USequencerCurveEditorObject::GetCustomColorForChannel(UClass* Class, const FString& Identifier)
+{
+	FLinearColor Color(FColor::White);
+	const UCurveEditorSettings* Settings = GetDefault<UCurveEditorSettings>();
+	if (Settings)
+	{
+		TOptional<FLinearColor> OptColor = Settings->GetCustomColor(Class, Identifier);
+		if (OptColor.IsSet())
+		{
+			return OptColor.GetValue();
+		}
+	}
+	return Color;
+}
+
+void USequencerCurveEditorObject::SetCustomColorForChannel(UClass* Class, const FString& Identifier, const FLinearColor& NewColor)
+{
+	UCurveEditorSettings* Settings = GetMutableDefault<UCurveEditorSettings>();
+	if (Settings)
+	{
+		Settings->SetCustomColor(Class, Identifier, NewColor);
+	}
+}
+
+void USequencerCurveEditorObject::SetCustomColorForChannels(UClass* Class, const TArray<FString>& Identifiers, const TArray<FLinearColor>& NewColors)
+{
+	if (Identifiers.Num() != NewColors.Num())
+	{
+		return;
+	}
+	UCurveEditorSettings* Settings = GetMutableDefault<UCurveEditorSettings>();
+	if (Settings)
+	{
+		for (int32 Index = 0; Index < Identifiers.Num(); ++Index)
+		{
+			const FString& Identifier = Identifiers[Index];
+			const FLinearColor& NewColor = NewColors[Index];
+			Settings->SetCustomColor(Class, Identifier, NewColor);
+		}
+	}
+}
+
+void USequencerCurveEditorObject::DeleteColorForChannels(UClass* Class, FString& Identifier)
+{
+	UCurveEditorSettings* Settings = GetMutableDefault<UCurveEditorSettings>();
+	if (Settings)
+	{
+		Settings->DeleteCustomColor(Class, Identifier);
+	}
+}
+
+void USequencerCurveEditorObject::SetRandomColorForChannels(UClass* Class, const TArray<FString>& Identifiers)
+{
+	UCurveEditorSettings* Settings = GetMutableDefault<UCurveEditorSettings>();
+	if (Settings)
+	{
+		for (int32 Index = 0; Index < Identifiers.Num(); ++Index)
+		{
+			const FString& Identifier = Identifiers[Index];
+			FLinearColor NewColor = UCurveEditorSettings::GetNextRandomColor();
+			Settings->SetCustomColor(Class, Identifier, NewColor);
+		}
+	}
+}
+
+
+
+
diff --git a/Engine/Plugins/MovieScene/SequencerScripting/Source/SequencerScriptingEditor/Public/SequencerCurveEditorObject.h b/Engine/Plugins/MovieScene/SequencerScripting/Source/SequencerScriptingEditor/Public/SequencerCurveEditorObject.h
new file mode 100644
--- /dev/null
+++ b/Engine/Plugins/MovieScene/SequencerScripting/Source/SequencerScriptingEditor/Public/SequencerCurveEditorObject.h
@@ -0,0 +1,129 @@
+// Copyright Epic Games, Inc. All Rights Reserved.
+
+#pragma once
+
+#include "Kismet/BlueprintFunctionLibrary.h"
+#include "SequencerCurveEditorObject.generated.h"
+
+
+class ISequencer;
+class ULevelSequence;
+class FCurveEditor;
+struct FMovieSceneChannel;
+class UMovieSceneSection;
+struct FCurveModelID;
+
+
+USTRUCT(BlueprintType)
+struct SEQUENCERSCRIPTINGEDITOR_API FSequencerChannelProxy
+{
+	GENERATED_BODY()
+
+	FSequencerChannelProxy()
+		: Section(nullptr)
+	{}
+
+	FSequencerChannelProxy(const FName& InChannelName, UMovieSceneSection* InSection)
+		: ChannelName(InChannelName)
+		, Section(InSection)
+	{}
+
+	UPROPERTY(BlueprintReadWrite, Category = Channel)
+	FName ChannelName;
+
+	UPROPERTY(BlueprintReadWrite, Category = Channel)
+	TObjectPtr<UMovieSceneSection> Section;
+};
+
+
+/*
+* Class to hold sequencer curve editor functions
+*/
+UCLASS()
+class SEQUENCERSCRIPTINGEDITOR_API USequencerCurveEditorObject : public UObject
+{
+public:
+
+	GENERATED_BODY()
+
+public:
+	/** Open curve editor*/
+	UFUNCTION(BlueprintCallable, Category = "Sequencer Curve Editor")
+	void OpenCurveEditor();
+
+	UFUNCTION(BlueprintPure, Category = "Sequencer Curve Editor")
+	/** Is curve editor open*/
+	bool IsCurveEditorOpen();
+
+	UFUNCTION(BlueprintCallable, Category = "Sequencer Curve Editor")
+	/** Close curve editor*/
+	void CloseCurveEditor();
+
+public:
+
+	/** Gets the channel with selected keys */
+	UFUNCTION(BlueprintPure, Category = "Sequencer Curve Editor")
+	TArray<FSequencerChannelProxy> GetChannelsWithSelectedKeys();
+
+	/** Gets the selected keys with this channel */
+	UFUNCTION(BlueprintPure, Category = "Sequencer Curve Editor")
+	TArray<int32> GetSelectedKeys(const FSequencerChannelProxy& ChannelProxy);
+
+	/** Select keys */
+	UFUNCTION(BlueprintCallable, Category = "Sequencer Curve Editor")
+	void SelectKeys(const FSequencerChannelProxy& Channel, const TArray<int32>& Indices);
+
+	/** Empties the current selection. */
+	UFUNCTION(BlueprintCallable, Category = "Sequencer Curve Editor")
+	void EmptySelection();
+
+public:
+
+	/** Get if a custom color for specified channel idendified by it's class and identifier exists */
+	UFUNCTION(BlueprintCallable, Category = "Sequencer Curve Editor")
+	bool HasCustomColorForChannel(UClass* Class, const FString& Identifier);
+
+	/** Get custom color for specified channel idendified by it's class and identifier,if none exists will return white*/
+	UFUNCTION(BlueprintCallable, Category = "Sequencer Curve Editor")
+	FLinearColor GetCustomColorForChannel(UClass* Class, const FString& Identifier);
+
+	/** Set Custom Color for specified channel idendified by it's class and identifier. This will be stored in editor user preferences.*/
+	UFUNCTION(BlueprintCallable, Category = "Sequencer Curve Editor")
+	void SetCustomColorForChannel(UClass* Class, const FString& Identifier, const FLinearColor& NewColor);
+
+	/** Set Custom Color for specified channels idendified by it's class and identifiers. This will be stored in editor user preferences.*/
+	UFUNCTION(BlueprintCallable, Category = "Sequencer Curve Editor")
+	void SetCustomColorForChannels(UClass* Class, const TArray<FString>& Identifiers, const TArray<FLinearColor>& NewColors);
+
+	/** Set Random Colors for specified channels idendified by it's class and identifiers. This will be stored in editor user preferences.*/
+	UFUNCTION(BlueprintCallable, Category = "Sequencer Curve Editor")
+	void SetRandomColorForChannels(UClass* Class, const TArray<FString>& Identifiers);
+
+	/** Delete for specified channel idendified by it's class and identifier.*/
+	UFUNCTION(BlueprintCallable, Category = "Sequencer Curve Editor")
+	void DeleteColorForChannels(UClass* Class, FString& Identifier);
+
+public:
+
+	/**
+	 * Function to assign a sequencer singleton.
+	 * NOTE: Only to be called by ULevelSequenceBlueprintLibrary::SetSequencer.
+	 */
+	void SetSequencer(TSharedPtr<ISequencer>& InSequencer);
+
+public:
+
+	/**
+	 * Utility function to get curve from a section and a name
+	 */
+	FCurveModelID GetCurve(UMovieSceneSection* InSection, const FName& InName);
+
+	/**
+	Utility function to get curve editor
+	*/
+	TSharedPtr<FCurveEditor> GetCurveEditor();
+
+private:
+	//internal sequencer
+	TWeakPtr<ISequencer> CurrentSequencer;
+};
diff --git a/Engine/Plugins/MovieScene/SequencerScripting/Source/SequencerScriptingEditor/SequencerScriptingEditor.Build.cs b/Engine/Plugins/MovieScene/SequencerScripting/Source/SequencerScriptingEditor/SequencerScriptingEditor.Build.cs
--- a/Engine/Plugins/MovieScene/SequencerScripting/Source/SequencerScriptingEditor/SequencerScriptingEditor.Build.cs
+++ b/Engine/Plugins/MovieScene/SequencerScripting/Source/SequencerScriptingEditor/SequencerScriptingEditor.Build.cs
@@ -47,7 +47,9 @@ public class SequencerScriptingEditor : ModuleRules
 				"EditorFramework",
                 "UnrealEd",
                 "Sequencer",
-                "BlueprintGraph"
+                "BlueprintGraph",
+				"CurveEditor",
+				"SequencerCore"
             }
 		);
 
diff --git a/Engine/Source/Editor/Sequencer/Public/MVVM/ViewModels/SectionModel.h b/Engine/Source/Editor/Sequencer/Public/MVVM/ViewModels/SectionModel.h
--- a/Engine/Source/Editor/Sequencer/Public/MVVM/ViewModels/SectionModel.h
+++ b/Engine/Source/Editor/Sequencer/Public/MVVM/ViewModels/SectionModel.h
@@ -32,7 +32,7 @@ struct FOverlappingSections;
 /**
  * Model for a sequencer section
  */
-class FSectionModel
+class SEQUENCER_API FSectionModel
 	: public FViewModel
 	, public FLinkedOutlinerExtension
 	, public ILayerBarExtension
diff --git a/Engine/Source/Runtime/MovieScene/Private/Channels/MovieSceneAudioTriggerChannel.cpp b/Engine/Source/Runtime/MovieScene/Private/Channels/MovieSceneAudioTriggerChannel.cpp
--- a/Engine/Source/Runtime/MovieScene/Private/Channels/MovieSceneAudioTriggerChannel.cpp
+++ b/Engine/Source/Runtime/MovieScene/Private/Channels/MovieSceneAudioTriggerChannel.cpp
@@ -97,6 +97,16 @@ void FMovieSceneAudioTriggerChannel::Optimize(const FKeyDataOptimizationParams&
 	UE::MovieScene::Optimize(this, InParameters);
 }
 
+FKeyHandle FMovieSceneAudioTriggerChannel::GetHandle(int32 Index)
+{
+	return GetData().GetHandle(Index);
+}
+
+int32 FMovieSceneAudioTriggerChannel::GetIndex(FKeyHandle Handle)
+{
+	return GetData().GetIndex(Handle);
+}
+
 void FMovieSceneAudioTriggerChannel::Offset(FFrameNumber DeltaPosition)
 {
 	GetData().Offset(DeltaPosition);
diff --git a/Engine/Source/Runtime/MovieScene/Private/Channels/MovieSceneBoolChannel.cpp b/Engine/Source/Runtime/MovieScene/Private/Channels/MovieSceneBoolChannel.cpp
--- a/Engine/Source/Runtime/MovieScene/Private/Channels/MovieSceneBoolChannel.cpp
+++ b/Engine/Source/Runtime/MovieScene/Private/Channels/MovieSceneBoolChannel.cpp
@@ -181,6 +181,16 @@ void FMovieSceneBoolChannel::Offset(FFrameNumber DeltaPosition)
 	GetData().Offset(DeltaPosition);
 }
 
+FKeyHandle FMovieSceneBoolChannel::GetHandle(int32 Index)
+{
+	return GetData().GetHandle(Index);
+}
+
+int32 FMovieSceneBoolChannel::GetIndex(FKeyHandle Handle)
+{
+	return GetData().GetIndex(Handle);
+}
+
 void FMovieSceneBoolChannel::ClearDefault()
 {
 	bHasDefaultValue = false;
diff --git a/Engine/Source/Runtime/MovieScene/Private/Channels/MovieSceneDoubleChannel.cpp b/Engine/Source/Runtime/MovieScene/Private/Channels/MovieSceneDoubleChannel.cpp
--- a/Engine/Source/Runtime/MovieScene/Private/Channels/MovieSceneDoubleChannel.cpp
+++ b/Engine/Source/Runtime/MovieScene/Private/Channels/MovieSceneDoubleChannel.cpp
@@ -158,6 +158,16 @@ void FMovieSceneDoubleChannel::Offset(FFrameNumber DeltaPosition)
 	AutoSetTangents();
 }
 
+FKeyHandle FMovieSceneDoubleChannel::GetHandle(int32 Index)
+{
+	return GetData().GetHandle(Index);
+}
+
+int32 FMovieSceneDoubleChannel::GetIndex(FKeyHandle Handle)
+{
+	return GetData().GetIndex(Handle);
+}
+
 void FMovieSceneDoubleChannel::Optimize(const FKeyDataOptimizationParams& Params)
 {
 	FMovieSceneDoubleChannelImpl::Optimize(this, Params);
diff --git a/Engine/Source/Runtime/MovieScene/Private/Channels/MovieSceneFloatChannel.cpp b/Engine/Source/Runtime/MovieScene/Private/Channels/MovieSceneFloatChannel.cpp
--- a/Engine/Source/Runtime/MovieScene/Private/Channels/MovieSceneFloatChannel.cpp
+++ b/Engine/Source/Runtime/MovieScene/Private/Channels/MovieSceneFloatChannel.cpp
@@ -148,6 +148,16 @@ void FMovieSceneFloatChannel::Offset(FFrameNumber DeltaPosition)
 	AutoSetTangents();
 }
 
+FKeyHandle FMovieSceneFloatChannel::GetHandle(int32 Index)
+{
+	return GetData().GetHandle(Index);
+}
+
+int32 FMovieSceneFloatChannel::GetIndex(FKeyHandle Handle)
+{
+	return GetData().GetIndex(Handle);
+}
+
 void FMovieSceneFloatChannel::Optimize(const FKeyDataOptimizationParams& Params)
 {
 	FMovieSceneFloatChannelImpl::Optimize(this, Params);
diff --git a/Engine/Source/Runtime/MovieScene/Private/Channels/MovieSceneIntegerChannel.cpp b/Engine/Source/Runtime/MovieScene/Private/Channels/MovieSceneIntegerChannel.cpp
--- a/Engine/Source/Runtime/MovieScene/Private/Channels/MovieSceneIntegerChannel.cpp
+++ b/Engine/Source/Runtime/MovieScene/Private/Channels/MovieSceneIntegerChannel.cpp
@@ -187,6 +187,16 @@ void FMovieSceneIntegerChannel::Offset(FFrameNumber DeltaPosition)
 	GetData().Offset(DeltaPosition);
 }
 
+FKeyHandle FMovieSceneIntegerChannel::GetHandle(int32 Index)
+{
+	return GetData().GetHandle(Index);
+}
+
+int32 FMovieSceneIntegerChannel::GetIndex(FKeyHandle Handle)
+{
+	return GetData().GetIndex(Handle);
+}
+
 void FMovieSceneIntegerChannel::ClearDefault()
 {
 	bHasDefaultValue = false;
diff --git a/Engine/Source/Runtime/MovieScene/Private/Channels/MovieSceneObjectPathChannel.cpp b/Engine/Source/Runtime/MovieScene/Private/Channels/MovieSceneObjectPathChannel.cpp
--- a/Engine/Source/Runtime/MovieScene/Private/Channels/MovieSceneObjectPathChannel.cpp
+++ b/Engine/Source/Runtime/MovieScene/Private/Channels/MovieSceneObjectPathChannel.cpp
@@ -155,6 +155,16 @@ void FMovieSceneObjectPathChannel::Offset(FFrameNumber DeltaPosition)
 	GetData().Offset(DeltaPosition);
 }
 
+FKeyHandle FMovieSceneObjectPathChannel::GetHandle(int32 Index)
+{
+	return GetData().GetHandle(Index);
+}
+
+int32 FMovieSceneObjectPathChannel::GetIndex(FKeyHandle Handle)
+{
+	return GetData().GetIndex(Handle);
+}
+
 void FMovieSceneObjectPathChannel::Optimize(const FKeyDataOptimizationParams& InParameters)
 {}
 
diff --git a/Engine/Source/Runtime/MovieScene/Public/Channels/MovieSceneAudioTriggerChannel.h b/Engine/Source/Runtime/MovieScene/Public/Channels/MovieSceneAudioTriggerChannel.h
--- a/Engine/Source/Runtime/MovieScene/Public/Channels/MovieSceneAudioTriggerChannel.h
+++ b/Engine/Source/Runtime/MovieScene/Public/Channels/MovieSceneAudioTriggerChannel.h
@@ -96,7 +96,8 @@ public:
 	MOVIESCENE_API virtual void Reset() override;
 	MOVIESCENE_API virtual void Offset(FFrameNumber DeltaPosition) override;
 	MOVIESCENE_API virtual void Optimize(const FKeyDataOptimizationParams& InParameters) override;
-
+	MOVIESCENE_API virtual FKeyHandle GetHandle(int32 Index) override;
+	MOVIESCENE_API virtual int32 GetIndex(FKeyHandle Handle) override;
 private:
 
 	UPROPERTY(meta = (KeyTimes))
diff --git a/Engine/Source/Runtime/MovieScene/Public/Channels/MovieSceneBoolChannel.h b/Engine/Source/Runtime/MovieScene/Public/Channels/MovieSceneBoolChannel.h
--- a/Engine/Source/Runtime/MovieScene/Public/Channels/MovieSceneBoolChannel.h
+++ b/Engine/Source/Runtime/MovieScene/Public/Channels/MovieSceneBoolChannel.h
@@ -122,6 +122,8 @@ public:
 	MOVIESCENE_API virtual void Offset(FFrameNumber DeltaPosition) override;
 	MOVIESCENE_API virtual void Optimize(const FKeyDataOptimizationParams& InParameters) override;
 	MOVIESCENE_API virtual void ClearDefault() override;
+	MOVIESCENE_API virtual FKeyHandle GetHandle(int32 Index) override;
+	MOVIESCENE_API virtual int32 GetIndex(FKeyHandle Handle) override;
 
 public:
 
diff --git a/Engine/Source/Runtime/MovieScene/Public/Channels/MovieSceneChannel.h b/Engine/Source/Runtime/MovieScene/Public/Channels/MovieSceneChannel.h
--- a/Engine/Source/Runtime/MovieScene/Public/Channels/MovieSceneChannel.h
+++ b/Engine/Source/Runtime/MovieScene/Public/Channels/MovieSceneChannel.h
@@ -10,6 +10,7 @@
 #include "Misc/CoreDefines.h"
 #include "Misc/FrameNumber.h"
 #include "Misc/FrameRate.h"
+#include "Curves/KeyHandle.h"
 #include "UObject/ObjectMacros.h"
 
 #include "MovieSceneChannel.generated.h"
@@ -17,7 +18,6 @@
 struct FFrameNumber;
 struct FFrameRate;
 struct FKeyDataOptimizationParams;
-struct FKeyHandle;
 struct FMovieSceneChannel;
 
 /*
@@ -241,6 +241,28 @@ struct FMovieSceneChannel
 	 */
 	virtual void PostEditChange() {}
 
+	/**
+	* Retrieve a key handle for the specified key time index
+	*
+	* @param Index          The index to retrieve
+	* @return A key handle that identifies the key at the specified index, regardless of re-ordering
+	*/
+	virtual FKeyHandle GetHandle(int32 Index)
+	{
+		return FKeyHandle();
+	}
+
+	/**
+	 * Attempt to retrieve the index of key from its handle
+	 *
+	 * @param Handle         The handle to retrieve
+	 * @return The index of the key, or INDEX_NONE
+	 */
+	virtual int32 GetIndex(FKeyHandle Handle)
+	{
+		return INDEX_NONE;
+	}
+
 public:
 
 	/**
diff --git a/Engine/Source/Runtime/MovieScene/Public/Channels/MovieSceneDoubleChannel.h b/Engine/Source/Runtime/MovieScene/Public/Channels/MovieSceneDoubleChannel.h
--- a/Engine/Source/Runtime/MovieScene/Public/Channels/MovieSceneDoubleChannel.h
+++ b/Engine/Source/Runtime/MovieScene/Public/Channels/MovieSceneDoubleChannel.h
@@ -212,6 +212,8 @@ public:
 	MOVIESCENE_API virtual void Optimize(const FKeyDataOptimizationParams& InParameters) override;
 	MOVIESCENE_API virtual void ClearDefault() override;
 	MOVIESCENE_API virtual void PostEditChange() override;
+	MOVIESCENE_API virtual FKeyHandle GetHandle(int32 Index) override;
+	MOVIESCENE_API virtual int32 GetIndex(FKeyHandle Handle) override;
 
 public:
 
diff --git a/Engine/Source/Runtime/MovieScene/Public/Channels/MovieSceneFloatChannel.h b/Engine/Source/Runtime/MovieScene/Public/Channels/MovieSceneFloatChannel.h
--- a/Engine/Source/Runtime/MovieScene/Public/Channels/MovieSceneFloatChannel.h
+++ b/Engine/Source/Runtime/MovieScene/Public/Channels/MovieSceneFloatChannel.h
@@ -206,6 +206,8 @@ public:
 	MOVIESCENE_API virtual void Optimize(const FKeyDataOptimizationParams& InParameters) override;
 	MOVIESCENE_API virtual void ClearDefault() override;
 	MOVIESCENE_API virtual void PostEditChange() override;
+	MOVIESCENE_API virtual FKeyHandle GetHandle(int32 Index) override;
+	MOVIESCENE_API virtual int32 GetIndex(FKeyHandle Handle) override;
 
 public:
 
diff --git a/Engine/Source/Runtime/MovieScene/Public/Channels/MovieSceneIntegerChannel.h b/Engine/Source/Runtime/MovieScene/Public/Channels/MovieSceneIntegerChannel.h
--- a/Engine/Source/Runtime/MovieScene/Public/Channels/MovieSceneIntegerChannel.h
+++ b/Engine/Source/Runtime/MovieScene/Public/Channels/MovieSceneIntegerChannel.h
@@ -145,6 +145,8 @@ public:
 	MOVIESCENE_API virtual void Offset(FFrameNumber DeltaPosition) override;
 	MOVIESCENE_API virtual void Optimize(const FKeyDataOptimizationParams& InParameters) override;
 	MOVIESCENE_API virtual void ClearDefault() override;
+	MOVIESCENE_API virtual FKeyHandle GetHandle(int32 Index) override;
+	MOVIESCENE_API virtual int32 GetIndex(FKeyHandle Handle) override;
 
 public:
 
diff --git a/Engine/Source/Runtime/MovieScene/Public/Channels/MovieSceneObjectPathChannel.h b/Engine/Source/Runtime/MovieScene/Public/Channels/MovieSceneObjectPathChannel.h
--- a/Engine/Source/Runtime/MovieScene/Public/Channels/MovieSceneObjectPathChannel.h
+++ b/Engine/Source/Runtime/MovieScene/Public/Channels/MovieSceneObjectPathChannel.h
@@ -131,6 +131,8 @@ public:
 	MOVIESCENE_API virtual void Offset(FFrameNumber DeltaPosition) override;
 	MOVIESCENE_API virtual void Optimize(const FKeyDataOptimizationParams& InParameters) override;
 	MOVIESCENE_API virtual void ClearDefault() override;
+	MOVIESCENE_API virtual FKeyHandle GetHandle(int32 Index) override;
+	MOVIESCENE_API virtual int32 GetIndex(FKeyHandle Handle) override;
 
 public:
 
diff --git a/Engine/Source/Runtime/MovieSceneTracks/Private/Channels/MovieSceneEventChannel.cpp b/Engine/Source/Runtime/MovieSceneTracks/Private/Channels/MovieSceneEventChannel.cpp
--- a/Engine/Source/Runtime/MovieSceneTracks/Private/Channels/MovieSceneEventChannel.cpp
+++ b/Engine/Source/Runtime/MovieSceneTracks/Private/Channels/MovieSceneEventChannel.cpp
@@ -61,3 +61,12 @@ void FMovieSceneEventChannel::Offset(FFrameNumber DeltaPosition)
 	GetData().Offset(DeltaPosition);
 }
 
+FKeyHandle FMovieSceneEventChannel::GetHandle(int32 Index)
+{
+	return GetData().GetHandle(Index);
+}
+
+int32 FMovieSceneEventChannel::GetIndex(FKeyHandle Handle)
+{
+	return GetData().GetIndex(Handle);
+}
diff --git a/Engine/Source/Runtime/MovieSceneTracks/Private/Channels/MovieSceneStringChannel.cpp b/Engine/Source/Runtime/MovieSceneTracks/Private/Channels/MovieSceneStringChannel.cpp
--- a/Engine/Source/Runtime/MovieSceneTracks/Private/Channels/MovieSceneStringChannel.cpp
+++ b/Engine/Source/Runtime/MovieSceneTracks/Private/Channels/MovieSceneStringChannel.cpp
@@ -92,6 +92,16 @@ void FMovieSceneStringChannel::DeleteKeysFrom(FFrameNumber InTime, bool bDeleteK
 	GetData().DeleteKeysFrom(InTime, bDeleteKeysBefore);
 }
 
+FKeyHandle FMovieSceneStringChannel::GetHandle(int32 Index)
+{
+	return GetData().GetHandle(Index);
+}
+
+int32 FMovieSceneStringChannel::GetIndex(FKeyHandle Handle)
+{
+	return GetData().GetIndex(Handle);
+}
+
 void FMovieSceneStringChannel::ChangeFrameResolution(FFrameRate SourceRate, FFrameRate DestinationRate)
 {
 	GetData().ChangeFrameResolution(SourceRate, DestinationRate);
diff --git a/Engine/Source/Runtime/MovieSceneTracks/Private/Sections/MovieSceneActorReferenceSection.cpp b/Engine/Source/Runtime/MovieSceneTracks/Private/Sections/MovieSceneActorReferenceSection.cpp
--- a/Engine/Source/Runtime/MovieSceneTracks/Private/Sections/MovieSceneActorReferenceSection.cpp
+++ b/Engine/Source/Runtime/MovieSceneTracks/Private/Sections/MovieSceneActorReferenceSection.cpp
@@ -60,6 +60,16 @@ void FMovieSceneActorReferenceData::DeleteKeysFrom(FFrameNumber InTime, bool bDe
 	GetData().DeleteKeysFrom(InTime, bDeleteKeysBefore);
 }
 
+FKeyHandle FMovieSceneActorReferenceData::GetHandle(int32 Index)
+{
+	return GetData().GetHandle(Index);
+}
+
+int32 FMovieSceneActorReferenceData::GetIndex(FKeyHandle Handle)
+{
+	return GetData().GetIndex(Handle);
+}
+
 void FMovieSceneActorReferenceData::ChangeFrameResolution(FFrameRate SourceRate, FFrameRate DestinationRate)
 {
 	GetData().ChangeFrameResolution(SourceRate, DestinationRate);
diff --git a/Engine/Source/Runtime/MovieSceneTracks/Private/Sections/MovieSceneEventSection.cpp b/Engine/Source/Runtime/MovieSceneTracks/Private/Sections/MovieSceneEventSection.cpp
--- a/Engine/Source/Runtime/MovieSceneTracks/Private/Sections/MovieSceneEventSection.cpp
+++ b/Engine/Source/Runtime/MovieSceneTracks/Private/Sections/MovieSceneEventSection.cpp
@@ -387,6 +387,17 @@ void FMovieSceneEventSectionData::Offset(FFrameNumber DeltaPosition)
 	GetData().Offset(DeltaPosition);
 }
 
+
+FKeyHandle FMovieSceneEventSectionData::GetHandle(int32 Index)
+{
+	return GetData().GetHandle(Index);
+}
+
+int32 FMovieSceneEventSectionData::GetIndex(FKeyHandle Handle)
+{
+	return GetData().GetIndex(Handle);
+}
+
 /* UMovieSceneSection structors
  *****************************************************************************/
 
diff --git a/Engine/Source/Runtime/MovieSceneTracks/Public/Channels/MovieSceneEventChannel.h b/Engine/Source/Runtime/MovieSceneTracks/Public/Channels/MovieSceneEventChannel.h
--- a/Engine/Source/Runtime/MovieSceneTracks/Public/Channels/MovieSceneEventChannel.h
+++ b/Engine/Source/Runtime/MovieSceneTracks/Public/Channels/MovieSceneEventChannel.h
@@ -60,7 +60,8 @@ public:
 	MOVIESCENETRACKS_API virtual int32 GetNumKeys() const override;
 	MOVIESCENETRACKS_API virtual void Reset() override;
 	MOVIESCENETRACKS_API virtual void Offset(FFrameNumber DeltaPosition) override;
-
+	MOVIESCENETRACKS_API virtual FKeyHandle GetHandle(int32 Index) override;
+	MOVIESCENETRACKS_API virtual int32 GetIndex(FKeyHandle Handle) override;
 private:
 
 	/** Array of times for each key */
diff --git a/Engine/Source/Runtime/MovieSceneTracks/Public/Channels/MovieSceneStringChannel.h b/Engine/Source/Runtime/MovieSceneTracks/Public/Channels/MovieSceneStringChannel.h
--- a/Engine/Source/Runtime/MovieSceneTracks/Public/Channels/MovieSceneStringChannel.h
+++ b/Engine/Source/Runtime/MovieSceneTracks/Public/Channels/MovieSceneStringChannel.h
@@ -112,6 +112,8 @@ public:
 	MOVIESCENETRACKS_API virtual void Offset(FFrameNumber DeltaPosition) override;
 	MOVIESCENETRACKS_API virtual void Optimize(const FKeyDataOptimizationParams& InParameters) override;
 	MOVIESCENETRACKS_API virtual void ClearDefault() override;
+	MOVIESCENETRACKS_API virtual FKeyHandle GetHandle(int32 Index) override;
+	MOVIESCENETRACKS_API virtual int32 GetIndex(FKeyHandle Handle) override;
 
 public:
 
diff --git a/Engine/Source/Runtime/MovieSceneTracks/Public/Sections/MovieSceneActorReferenceSection.h b/Engine/Source/Runtime/MovieSceneTracks/Public/Sections/MovieSceneActorReferenceSection.h
--- a/Engine/Source/Runtime/MovieSceneTracks/Public/Sections/MovieSceneActorReferenceSection.h
+++ b/Engine/Source/Runtime/MovieSceneTracks/Public/Sections/MovieSceneActorReferenceSection.h
@@ -121,6 +121,8 @@ public:
 	MOVIESCENETRACKS_API virtual void Reset() override;
 	MOVIESCENETRACKS_API virtual void Offset(FFrameNumber DeltaPosition) override;
 	MOVIESCENETRACKS_API virtual void ClearDefault() override;
+	MOVIESCENETRACKS_API virtual FKeyHandle GetHandle(int32 Index) override;
+	MOVIESCENETRACKS_API virtual int32 GetIndex(FKeyHandle Handle) override;
 
 public:
 
diff --git a/Engine/Source/Runtime/MovieSceneTracks/Public/Sections/MovieSceneEventSection.h b/Engine/Source/Runtime/MovieSceneTracks/Public/Sections/MovieSceneEventSection.h
--- a/Engine/Source/Runtime/MovieSceneTracks/Public/Sections/MovieSceneEventSection.h
+++ b/Engine/Source/Runtime/MovieSceneTracks/Public/Sections/MovieSceneEventSection.h
@@ -184,6 +184,8 @@ public:
 	MOVIESCENETRACKS_API virtual int32 GetNumKeys() const override;
 	MOVIESCENETRACKS_API virtual void Reset() override;
 	MOVIESCENETRACKS_API virtual void Offset(FFrameNumber DeltaPosition) override;
+	MOVIESCENETRACKS_API virtual FKeyHandle GetHandle(int32 Index) override;
+	MOVIESCENETRACKS_API virtual int32 GetIndex(FKeyHandle Handle) override;
 	virtual void Optimize(const FKeyDataOptimizationParams& InParameters) override {}
 	virtual void ClearDefault() override {}
 
