Sequencer: Curve Editor Scripting: Make sure to use long intention name

diff --git a/Engine/Plugins/MovieScene/SequencerScripting/Source/SequencerScriptingEditor/Private/SequencerCurveEditorObject.cpp b/Engine/Plugins/MovieScene/SequencerScripting/Source/SequencerScriptingEditor/Private/SequencerCurveEditorObject.cpp
--- a/Engine/Plugins/MovieScene/SequencerScripting/Source/SequencerScriptingEditor/Private/SequencerCurveEditorObject.cpp
+++ b/Engine/Plugins/MovieScene/SequencerScripting/Source/SequencerScriptingEditor/Private/SequencerCurveEditorObject.cpp
@@ -1,3 +1,4 @@
+
 // Copyright Epic Games, Inc. All Rights Reserved.
 #include "SequencerCurveEditorObject.h"
 #include "MVVM/CurveEditorExtension.h"
@@ -80,10 +81,8 @@ TArray<FSequencerChannelProxy> USequencerCurveEditorObject::GetChannelsWithSelec
 			{
 				if (UMovieSceneSection* Section = Cast<UMovieSceneSection>(Curve->GetOwningObject()))
 				{
-					FString IntentionName = Curve->GetIntentionName();
-					FText DisplayName = Curve->GetShortDisplayName();
-					FText LongName = Curve->GetLongDisplayName();
-					FName ChannelName = FName(*DisplayName.ToString());
+					FString IntentionName = Curve->GetLongIntentionName();
+					FName ChannelName = FName(*IntentionName);
 					FSequencerChannelProxy ChannelProxy(ChannelName,Section);
 					OutSelectedChannels.Add(ChannelProxy);
 
