Sequencer: Scripting: Unify channel names to use the default for both curves and channels, by adding 'ChannelName' property to curves and using that in scripting.

diff --git a/Engine/Plugins/MovieScene/SequencerScripting/Source/SequencerScriptingEditor/Private/SequencerCurveEditorObject.cpp b/Engine/Plugins/MovieScene/SequencerScripting/Source/SequencerScriptingEditor/Private/SequencerCurveEditorObject.cpp
--- a/Engine/Plugins/MovieScene/SequencerScripting/Source/SequencerScriptingEditor/Private/SequencerCurveEditorObject.cpp
+++ b/Engine/Plugins/MovieScene/SequencerScripting/Source/SequencerScriptingEditor/Private/SequencerCurveEditorObject.cpp
@@ -81,8 +81,7 @@ TArray<FSequencerChannelProxy> USequencerCurveEditorObject::GetChannelsWithSelec
 			{
 				if (UMovieSceneSection* Section = Cast<UMovieSceneSection>(Curve->GetOwningObject()))
 				{
-					FString IntentionName = Curve->GetLongIntentionName();
-					FName ChannelName = FName(*IntentionName);
+					FName ChannelName = Curve->GetChannelName();
 					FSequencerChannelProxy ChannelProxy(ChannelName,Section);
 					OutSelectedChannels.Add(ChannelProxy);
 
@@ -146,7 +145,7 @@ FCurveModelID USequencerCurveEditorObject::GetCurve(UMovieSceneSection* InSectio
 		const TMap<FCurveModelID, TUniquePtr<FCurveModel>>& Curves = CurveEditor->GetCurves();
 		for (const TPair <FCurveModelID, TUniquePtr<FCurveModel>>& Pair : Curves)
 		{
-			if (Pair.Value.IsValid() && Pair.Value->GetOwningObject() == InSection && Pair.Value->GetLongIntentionName() == InName)
+			if (Pair.Value.IsValid() && Pair.Value->GetOwningObject() == InSection && Pair.Value->GetChannelName() == InName)
 			{
 				return Pair.Key;
 			}
diff --git a/Engine/Source/Editor/CurveEditor/Public/CurveModel.h b/Engine/Source/Editor/CurveEditor/Public/CurveModel.h
--- a/Engine/Source/Editor/CurveEditor/Public/CurveModel.h
+++ b/Engine/Source/Editor/CurveEditor/Public/CurveModel.h
@@ -384,6 +384,16 @@ public:
 		return LongIntentionName;
 	}
 
+	FORCEINLINE void SetChannelName(const FName& InChannelName)
+	{
+		ChannelName = InChannelName;
+	}
+
+	FORCEINLINE FName GetChannelName() const
+	{
+		return ChannelName;
+	}
+
 	/**
 	 */
 	FORCEINLINE void SetColor(const FLinearColor& InColor, bool bInModify = true)
@@ -436,6 +446,11 @@ protected:
 	*/
 	FString LongIntentionName;
 
+	/**
+	* The original channel name, used mostly to make sure names match with BP/Scripting
+	*/
+	FName ChannelName;
+
 	/** This curve's display color */
 	FLinearColor Color;
 
diff --git a/Engine/Source/Editor/Sequencer/Private/IKeyArea.cpp b/Engine/Source/Editor/Sequencer/Private/IKeyArea.cpp
--- a/Engine/Source/Editor/Sequencer/Private/IKeyArea.cpp
+++ b/Engine/Source/Editor/Sequencer/Private/IKeyArea.cpp
@@ -329,6 +329,7 @@ TUniquePtr<FCurveModel> IKeyArea::CreateCurveEditorModel(TSharedRef<ISequencer>
 			CurveModel->SetLongDisplayName(LongDisplayName);
 			CurveModel->SetIntentionName(IntentName);
 			CurveModel->SetLongIntentionName(LongIntentName);
+			CurveModel->SetChannelName(ChannelHandle.GetMetaData()->Name);
 			if (Color.IsSet())
 			{
 				CurveModel->SetColor(Color.GetValue(),false);
