Sequencer: Add time unit to marked frame commands

Deprecate existing functions: GetMarkedFrames, AddMarkedFrame, SetMarkedFrame, FindMarkedFrameByFrameNumber, FindNextMarkedFrame
Added new functions that have the old function's display names: GetMarkedFramesInSequence, AddMarkedFrameToSequence, SetMarkedFrameInSequence, FindMarkedFrameByFrameNumberInSequence, FindNextMarkedFrameInSequence

diff --git a/Engine/Plugins/MovieScene/SequencerScripting/Source/SequencerScripting/Private/ExtensionLibraries/MovieSceneSequenceExtensions.cpp b/Engine/Plugins/MovieScene/SequencerScripting/Source/SequencerScripting/Private/ExtensionLibraries/MovieSceneSequenceExtensions.cpp
--- a/Engine/Plugins/MovieScene/SequencerScripting/Source/SequencerScripting/Private/ExtensionLibraries/MovieSceneSequenceExtensions.cpp
+++ b/Engine/Plugins/MovieScene/SequencerScripting/Source/SequencerScripting/Private/ExtensionLibraries/MovieSceneSequenceExtensions.cpp
@@ -1031,7 +1031,7 @@ void UMovieSceneSequenceExtensions::RemoveRootFolderFromSequence(UMovieSceneSequ
 #endif
 }
 
-TArray<FMovieSceneMarkedFrame> UMovieSceneSequenceExtensions::GetMarkedFrames(UMovieSceneSequence* Sequence)
+TArray<FMovieSceneMarkedFrame> UMovieSceneSequenceExtensions::GetMarkedFramesFromSequence(UMovieSceneSequence* Sequence, ESequenceTimeUnit TimeUnit)
 {
 	if (!Sequence)
 	{
@@ -1042,13 +1042,25 @@ TArray<FMovieSceneMarkedFrame> UMovieSceneSequenceExtensions::GetMarkedFrames(UM
 	UMovieScene* MovieScene = Sequence->GetMovieScene();
 	if (MovieScene)
 	{
-		return MovieScene->GetMarkedFrames();
+		TArray<FMovieSceneMarkedFrame> MarkedFrames = MovieScene->GetMarkedFrames();
+		if (TimeUnit == ESequenceTimeUnit::DisplayRate)
+		{
+			FFrameRate TickResolution = MovieScene->GetTickResolution();
+			FFrameRate DisplayRate = MovieScene->GetDisplayRate();
+
+			for (FMovieSceneMarkedFrame& MarkedFrame : MarkedFrames)
+			{
+				MarkedFrame.FrameNumber = FFrameRate::TransformTime(MarkedFrame.FrameNumber, TickResolution, DisplayRate).RoundToFrame();
+			}
+		}
+
+		return MarkedFrames;
 	}
 
 	return TArray<FMovieSceneMarkedFrame>();
 }
 
-int32 UMovieSceneSequenceExtensions::AddMarkedFrame(UMovieSceneSequence* Sequence, const FMovieSceneMarkedFrame& InMarkedFrame)
+int32 UMovieSceneSequenceExtensions::AddMarkedFrameToSequence(UMovieSceneSequence* Sequence, const FMovieSceneMarkedFrame& InMarkedFrame, ESequenceTimeUnit TimeUnit)
 {
 	if (!Sequence)
 	{
@@ -1061,13 +1073,19 @@ int32 UMovieSceneSequenceExtensions::AddMarkedFrame(UMovieSceneSequence* Sequenc
 	{
 		MovieScene->Modify();
 
-		return MovieScene->AddMarkedFrame(InMarkedFrame);
+		FMovieSceneMarkedFrame MarkedFrame(InMarkedFrame);
+		if (TimeUnit == ESequenceTimeUnit::DisplayRate)
+		{
+			MarkedFrame.FrameNumber = FFrameRate::TransformTime(MarkedFrame.FrameNumber, MovieScene->GetDisplayRate(), MovieScene->GetTickResolution()).RoundToFrame();
+		}
+
+		return MovieScene->AddMarkedFrame(MarkedFrame);
 	}
 	return INDEX_NONE;
 }
 
 
-void UMovieSceneSequenceExtensions::SetMarkedFrame(UMovieSceneSequence* Sequence, int32 InMarkIndex, FFrameNumber InFrameNumber)
+void UMovieSceneSequenceExtensions::SetMarkedFrameInSequence(UMovieSceneSequence* Sequence, int32 InMarkIndex, FFrameNumber InFrameNumber, ESequenceTimeUnit TimeUnit)
 {
 	if (!Sequence)
 	{
@@ -1078,6 +1096,11 @@ void UMovieSceneSequenceExtensions::SetMarkedFrame(UMovieSceneSequence* Sequence
 	UMovieScene* MovieScene = Sequence->GetMovieScene();
 	if (MovieScene)
 	{
+		if (TimeUnit == ESequenceTimeUnit::DisplayRate)
+		{
+			InFrameNumber = FFrameRate::TransformTime(InFrameNumber, MovieScene->GetDisplayRate(), MovieScene->GetTickResolution()).RoundToFrame();
+		}
+
 		return MovieScene->SetMarkedFrame(InMarkIndex, InFrameNumber);
 	}
 }
@@ -1148,7 +1171,7 @@ int32 UMovieSceneSequenceExtensions::FindMarkedFrameByLabel(UMovieSceneSequence*
 	return INDEX_NONE;
 }
 
-int32 UMovieSceneSequenceExtensions::FindMarkedFrameByFrameNumber(UMovieSceneSequence* Sequence, FFrameNumber InFrameNumber)
+int32 UMovieSceneSequenceExtensions::FindMarkedFrameByFrameNumberInSequence(UMovieSceneSequence* Sequence, FFrameNumber InFrameNumber, ESequenceTimeUnit TimeUnit)
 {
 	if (!Sequence)
 	{
@@ -1159,12 +1182,17 @@ int32 UMovieSceneSequenceExtensions::FindMarkedFrameByFrameNumber(UMovieSceneSeq
 	UMovieScene* MovieScene = Sequence->GetMovieScene();
 	if (MovieScene)
 	{
+		if (TimeUnit == ESequenceTimeUnit::DisplayRate)
+		{
+			InFrameNumber = FFrameRate::TransformTime(InFrameNumber, MovieScene->GetDisplayRate(), MovieScene->GetTickResolution()).RoundToFrame();
+		}
+
 		return MovieScene->FindMarkedFrameByFrameNumber(InFrameNumber);
 	}
 	return INDEX_NONE;
 }
 
-int32 UMovieSceneSequenceExtensions::FindNextMarkedFrame(UMovieSceneSequence* Sequence, FFrameNumber InFrameNumber, bool bForward)
+int32 UMovieSceneSequenceExtensions::FindNextMarkedFrameInSequence(UMovieSceneSequence* Sequence, FFrameNumber InFrameNumber, bool bForward, ESequenceTimeUnit TimeUnit)
 {
 	if (!Sequence)
 	{
@@ -1175,6 +1203,11 @@ int32 UMovieSceneSequenceExtensions::FindNextMarkedFrame(UMovieSceneSequence* Se
 	UMovieScene* MovieScene = Sequence->GetMovieScene();
 	if (MovieScene)
 	{
+		if (TimeUnit == ESequenceTimeUnit::DisplayRate)
+		{
+			InFrameNumber = FFrameRate::TransformTime(InFrameNumber, MovieScene->GetDisplayRate(), MovieScene->GetTickResolution()).RoundToFrame();
+		}
+
 		return MovieScene->FindNextMarkedFrame(InFrameNumber, bForward);
 	}
 	return INDEX_NONE;
diff --git a/Engine/Plugins/MovieScene/SequencerScripting/Source/SequencerScripting/Public/ExtensionLibraries/MovieSceneSequenceExtensions.h b/Engine/Plugins/MovieScene/SequencerScripting/Source/SequencerScripting/Public/ExtensionLibraries/MovieSceneSequenceExtensions.h
--- a/Engine/Plugins/MovieScene/SequencerScripting/Source/SequencerScripting/Public/ExtensionLibraries/MovieSceneSequenceExtensions.h
+++ b/Engine/Plugins/MovieScene/SequencerScripting/Source/SequencerScripting/Public/ExtensionLibraries/MovieSceneSequenceExtensions.h
@@ -5,6 +5,7 @@
 #include "Kismet/BlueprintFunctionLibrary.h"
 #include "MovieSceneObjectBindingID.h" // for EMovieSceneObjectBindingSpace
 #include "MovieSceneTrack.h"
+#include "SequenceTimeUnit.h"
 #include "MovieSceneSequenceExtensions.generated.h"
 
 class UMovieScene;
@@ -528,10 +529,16 @@ public:
 public:
 
 	/*
+	 * Get the marked frames for this sequence
+	 *
 	 * @return Return the user marked frames
 	 */
-	UFUNCTION(BlueprintCallable, Category = "Sequencer|Sequence", meta = (ScriptMethod))
-	static TArray<FMovieSceneMarkedFrame> GetMarkedFrames(UMovieSceneSequence* Sequence);
+	UFUNCTION(BlueprintCallable, Category = "Sequencer|Sequence", meta = (ScriptMethod, DisplayName = "Get Marked Frames"))
+	static TArray<FMovieSceneMarkedFrame> GetMarkedFramesFromSequence(UMovieSceneSequence* Sequence, ESequenceTimeUnit TimeUnit = ESequenceTimeUnit::DisplayRate);
+
+	UE_DEPRECATED(5.4, "GetMarkedFrames is deprecated. Please use GetMarkedFrames that takes a time unit and defaults to display rate instead")
+	UFUNCTION(BlueprintCallable, Category = "Sequencer|Sequence", meta = (ScriptMethod, DeprecatedFunction, DeprecationMessage = "GetMarkedFrames is deprecated. Please use GetMarkedFrames that takes a time unit instead"))
+	static TArray<FMovieSceneMarkedFrame> GetMarkedFrames(UMovieSceneSequence* Sequence) { return GetMarkedFramesFromSequence(Sequence, ESequenceTimeUnit::TickResolution); }
 
 	/*
 	 * Add a given user marked frame.
@@ -540,8 +547,12 @@ public:
 	 * @InMarkedFrame The given user marked frame to add
 	 * @return The index to the newly added marked frame
 	 */
-	UFUNCTION(BlueprintCallable, Category = "Sequencer|Sequence", meta = (ScriptMethod))
-	static int32 AddMarkedFrame(UMovieSceneSequence* Sequence, const FMovieSceneMarkedFrame& InMarkedFrame);
+	UFUNCTION(BlueprintCallable, Category = "Sequencer|Sequence", meta = (ScriptMethod, DisplayName = "Add Marked Frame"))
+	static int32 AddMarkedFrameToSequence(UMovieSceneSequence* Sequence, const FMovieSceneMarkedFrame& InMarkedFrame, ESequenceTimeUnit TimeUnit = ESequenceTimeUnit::DisplayRate);
+
+	UE_DEPRECATED(5.4, "AddMarkedFrame is deprecated. Please use AddMarkedFrame that takes a time unit and defaults to display rate instead")
+	UFUNCTION(BlueprintCallable, Category = "Sequencer|Sequence", meta = (ScriptMethod, DeprecatedFunction, DeprecationMessage = "AddMarkedFrame is deprecated. Please use AddMarkedFrame that takes a time unit instead"))
+	static int32 AddMarkedFrame(UMovieSceneSequence* Sequence, const FMovieSceneMarkedFrame& InMarkedFrame) { return AddMarkedFrameToSequence(Sequence, InMarkedFrame, ESequenceTimeUnit::TickResolution); }
 
 	/*
 	 * Sets the frame number for the given marked frame index. Does not maintain sort. Call SortMarkedFrames
@@ -549,8 +560,12 @@ public:
 	 * @InMarkIndex The given user marked frame index to edit
 	 * @InFrameNumber The frame number to set
 	 */
-	UFUNCTION(BlueprintCallable, Category = "Sequencer|Sequence", meta = (ScriptMethod))
-	static void SetMarkedFrame(UMovieSceneSequence* Sequence, int32 InMarkIndex, FFrameNumber InFrameNumber);
+	UFUNCTION(BlueprintCallable, Category = "Sequencer|Sequence", meta = (ScriptMethod, DisplayName = "Set Marked Frame"))
+	static void SetMarkedFrameInSequence(UMovieSceneSequence* Sequence, int32 InMarkIndex, FFrameNumber InFrameNumber, ESequenceTimeUnit TimeUnit = ESequenceTimeUnit::DisplayRate);
+
+	UE_DEPRECATED(5.4, "SetMarkedFrame is deprecated. Please use SetMarkedFrame that takes a time unit and defaults to display rate instead")
+	UFUNCTION(BlueprintCallable, Category = "Sequencer|Sequence", meta = (ScriptMethod, DeprecatedFunction, DeprecationMessage = "SetMarkedFrame is deprecated. Please use SetMarkedFrame that takes a time unit instead"))
+	static void SetMarkedFrame(UMovieSceneSequence* Sequence, int32 InMarkIndex, FFrameNumber InFrameNumber) { SetMarkedFrameInSequence(Sequence, InMarkIndex, InFrameNumber, ESequenceTimeUnit::TickResolution); }
 
 	/*
 	 * Delete the user marked frame by index.
@@ -585,8 +600,12 @@ public:
 	 *
 	 * @InFrameNumber The frame number of the user marked frame to find
 	 */
-	UFUNCTION(BlueprintCallable, Category = "Sequencer|Sequence", meta = (ScriptMethod))
-	static int32 FindMarkedFrameByFrameNumber(UMovieSceneSequence* Sequence, FFrameNumber InFrameNumber);
+	UFUNCTION(BlueprintCallable, Category = "Sequencer|Sequence", meta = (ScriptMethod, DisplayName = "Find Marked Frame By Frame Number"))
+	static int32 FindMarkedFrameByFrameNumberInSequence(UMovieSceneSequence* Sequence, FFrameNumber InFrameNumber, ESequenceTimeUnit TimeUnit = ESequenceTimeUnit::DisplayRate);
+
+	UE_DEPRECATED(5.4, "FindMarkedFrameByFrameNumber is deprecated. Please use FindMarkedFrameByFrameNumber that takes a time unit and defaults to display rate instead")
+	UFUNCTION(BlueprintCallable, Category = "Sequencer|Sequence", meta = (ScriptMethod, DeprecatedFunction, DeprecationMessage = "FindMarkedFrameByFrameNumber is deprecated. Please use FindMarkedFrameByFrameNumber that takes a time unit instead"))
+	static int32 FindMarkedFrameByFrameNumber(UMovieSceneSequence* Sequence, FFrameNumber InFrameNumber) { return FindMarkedFrameByFrameNumberInSequence(Sequence, InFrameNumber, ESequenceTimeUnit::TickResolution); }
 
 	/*
 	 * Find the next/previous user marked frame from the given frame number
@@ -594,8 +613,12 @@ public:
 	 * @InFrameNumber The frame number to find the next/previous user marked frame from
 	 * @bForward Find forward from the given frame number.
 	 */
-	UFUNCTION(BlueprintCallable, Category = "Sequencer|Sequence", meta = (ScriptMethod))
-	static int32 FindNextMarkedFrame(UMovieSceneSequence* Sequence, FFrameNumber InFrameNumber, bool bForward);
+	UFUNCTION(BlueprintCallable, Category = "Sequencer|Sequence", meta = (ScriptMethod, DisplayName = "Find Next Marked Frame"))
+	static int32 FindNextMarkedFrameInSequence(UMovieSceneSequence* Sequence, FFrameNumber InFrameNumber, bool bForward, ESequenceTimeUnit TimeUnit = ESequenceTimeUnit::DisplayRate);
+
+	UE_DEPRECATED(5.4, "FindNextMarkedFrame is deprecated. Please use FindNextMarkedFrame that takes a time unit instead")
+	UFUNCTION(BlueprintCallable, Category = "Sequencer|Sequence", meta = (ScriptMethod, DeprecatedFunction, DeprecationMessage = "FindNextMarkedFrame is deprecated. Please use FindNextMarkedFrame that takes a time unit and defaults to display rate instead"))
+	static int32 FindNextMarkedFrame(UMovieSceneSequence* Sequence, FFrameNumber InFrameNumber, bool bForward) { return FindNextMarkedFrameInSequence(Sequence, InFrameNumber, bForward, ESequenceTimeUnit::TickResolution); }
 
 	/*
 	 * Set read only
