Sequencer: Follow up to 30931538

diff --git a/Engine/Plugins/MovieScene/SequencerScripting/Source/SequencerScripting/Public/SequenceTimeUnit.h b/Engine/Plugins/MovieScene/SequencerScripting/Source/SequencerScripting/Public/SequenceTimeUnit.h
--- a/Engine/Plugins/MovieScene/SequencerScripting/Source/SequencerScripting/Public/SequenceTimeUnit.h
+++ b/Engine/Plugins/MovieScene/SequencerScripting/Source/SequencerScripting/Public/SequenceTimeUnit.h
@@ -4,7 +4,7 @@
 
 #include "SequenceTimeUnit.generated.h"
 
-UENUM(BlueprintType)
+UENUM()
 enum class UE_DEPRECATED(5.4, "ESequenceTimeUnit is deprecated, please use EMovieSceneTimeUnit") ESequenceTimeUnit : uint8
 {
 	DisplayRate,
