Seqencer: Add Evaluate to Integer channel, form Matt Ringot

diff --git a/Engine/Plugins/MovieScene/SequencerScripting/Source/SequencerScripting/Private/KeysAndChannels/MovieSceneScriptingInteger.h b/Engine/Plugins/MovieScene/SequencerScripting/Source/SequencerScripting/Private/KeysAndChannels/MovieSceneScriptingInteger.h
--- a/Engine/Plugins/MovieScene/SequencerScripting/Source/SequencerScripting/Private/KeysAndChannels/MovieSceneScriptingInteger.h
+++ b/Engine/Plugins/MovieScene/SequencerScripting/Source/SequencerScripting/Private/KeysAndChannels/MovieSceneScriptingInteger.h
@@ -115,6 +115,26 @@ public:
 		return Impl::GetKeysInChannelByIndex(ChannelHandle, OwningSequence, OwningSection, Indices);
 	}
 
+	/**
+	* Returns number of keys in this channel.
+	*/
+	UFUNCTION(BlueprintCallable, Category = "Sequencer|Keys", meta = (DisplayName = "Get Num Keys (Integer)"))
+	int32 GetNumKeys() const
+	{
+		return ChannelHandle.Get() ? ChannelHandle.Get()->GetNumKeys() : 0;
+	}
+
+	/**
+	* Gets baked keys in this channel.
+	* @return	An array of values contained by this channel.
+	*			Returns baked keys in the specified range.
+	*/
+	UFUNCTION(BlueprintCallable, Category = "Sequencer|Keys", meta = (DisplayName = "Evaluate Keys (Integer)"))
+	TArray<int32> EvaluateKeys(FSequencerScriptingRange Range, FFrameRate FrameRate) const
+	{
+		return Impl::EvaluateKeysInChannel(ChannelHandle, OwningSequence, Range, FrameRate);
+	}
+
 	/**
 	* Set this channel's default value that should be used when no keys are present.
 	* Sets bHasDefaultValue to true automatically.
