Sequencer: Curve Editor: New Gaussian filter.  Added BP/Python functions to set/get curve display in the curve editor, curve filters are now blueprintable, and you can apply the filter.

diff --git a/Engine/Plugins/Animation/ControlRig/Source/ControlRigEditor/Private/EditMode/AnimDetailsProxy.cpp b/Engine/Plugins/Animation/ControlRig/Source/ControlRigEditor/Private/EditMode/AnimDetailsProxy.cpp
--- a/Engine/Plugins/Animation/ControlRig/Source/ControlRigEditor/Private/EditMode/AnimDetailsProxy.cpp
+++ b/Engine/Plugins/Animation/ControlRig/Source/ControlRigEditor/Private/EditMode/AnimDetailsProxy.cpp
@@ -3751,6 +3751,13 @@ bool UControlRigDetailPanelControlProxies::SelectPropertyInternal(UControlRigCon
 						{
 							if (TrackModel->GetTrack() == Element.WeakTrack.Get())
 							{
+								if (TViewModelPtr<FChannelGroupOutlinerModel> ChannelModel = CastViewModel<FChannelGroupOutlinerModel>(OutlinerExtenstionIt.GetCurrentItem()))
+								{
+									if (ChannelModel->GetChannel(TrackModel->GetTrack()->GetSectionToKey()) == nullptr) //if not section to key we also don't select it.
+									{
+										continue;
+									}
+								}
 								FName ID = OutlinerExtenstionIt->GetIdentifier();
 								FString Name = ID.ToString();
 								TArray<FString> StringArray;
diff --git a/Engine/Plugins/MovieScene/SequencerScripting/Source/SequencerScriptingEditor/Private/SequencerCurveEditorObject.cpp b/Engine/Plugins/MovieScene/SequencerScripting/Source/SequencerScriptingEditor/Private/SequencerCurveEditorObject.cpp
--- a/Engine/Plugins/MovieScene/SequencerScripting/Source/SequencerScriptingEditor/Private/SequencerCurveEditorObject.cpp
+++ b/Engine/Plugins/MovieScene/SequencerScripting/Source/SequencerScriptingEditor/Private/SequencerCurveEditorObject.cpp
@@ -3,12 +3,20 @@
 #include "SequencerCurveEditorObject.h"
 #include "MVVM/CurveEditorExtension.h"
 #include "MVVM/ViewModels/SequencerEditorViewModel.h"
+#include "MVVM/ViewModels/ChannelModel.h"
+#include "MVVM/SectionModelStorageExtension.h"
+#include "MVVM/ViewModels/OutlinerViewModel.h"
+#include "MVVM/Extensions/IOutlinerExtension.h"
+#include "MVVM/ViewModels/TrackModel.h"
+#include "Tree/SCurveEditorTree.h"
 #include "ISequencer.h"
 #include "CurveEditor.h"
 #include "CurveEditorTypes.h"
 #include "CurveModel.h"
 #include "Channels/MovieSceneChannel.h"
 #include "ExtensionLibraries/MovieSceneSectionExtensions.h"
+#include "Filters/CurveEditorFilterBase.h"
+
 //For custom colors on channels, stored in editor pref's
 #include "CurveEditorSettings.h"
 
@@ -137,8 +145,65 @@ void USequencerCurveEditorObject::EmptySelection()
 	}
 }
 
-FCurveModelID USequencerCurveEditorObject::GetCurve(UMovieSceneSection* InSection, const FName& InName)
+void USequencerCurveEditorObject::ShowCurve(const FSequencerChannelProxy& ChannelProxy, bool bShowCurve)
 {
+	using namespace UE::Sequencer;
+	TSharedPtr<FCurveEditor> CurveEditor = GetCurveEditor();
+	if (CurrentSequencer.IsValid() && CurveEditor.IsValid())
+	{
+		if (IsCurveShown(ChannelProxy) != bShowCurve)
+		{
+			const TSharedPtr<FSequencerEditorViewModel> SequencerViewModel = CurrentSequencer.Pin()->GetViewModel();
+			const FCurveEditorExtension* CurveEditorExtension = SequencerViewModel->CastDynamic<FCurveEditorExtension>();
+			check(CurveEditorExtension);
+			TSharedPtr<SCurveEditorTree>  CurveEditorTreeView = CurveEditorExtension->GetCurveEditorTreeView();
+			TSharedPtr<FOutlinerViewModel> OutlinerViewModel = SequencerViewModel->GetOutliner();
+			bool bIsSelected = false;
+			TParentFirstChildIterator<IOutlinerExtension> OutlinerExtenstionIt = OutlinerViewModel->GetRootItem()->GetDescendantsOfType<IOutlinerExtension>();
+			for (; OutlinerExtenstionIt; ++OutlinerExtenstionIt)
+			{
+				if (TSharedPtr<FTrackModel> TrackModel = OutlinerExtenstionIt.GetCurrentItem()->FindAncestorOfType<FTrackModel>())
+				{
+					if (UMovieSceneTrack* Track = TrackModel->GetTrack())
+					{
+						if (TViewModelPtr<FChannelGroupOutlinerModel> ChannelModel = CastViewModel<FChannelGroupOutlinerModel>(OutlinerExtenstionIt.GetCurrentItem()))
+						{
+							if (TSharedPtr<FChannelModel> ChannelPtr = ChannelModel->GetChannel(ChannelProxy.Section)) //if not section to key we also don't select it.
+							{
+								if (ChannelPtr->GetChannelName() == ChannelProxy.ChannelName)
+								{
+									if (TViewModelPtr<ICurveEditorTreeItemExtension> CurveEditorItem = OutlinerExtenstionIt.GetCurrentItem().ImplicitCast())
+									{
+										FCurveEditorTreeItemID CurveEditorTreeItem = CurveEditorItem->GetCurveEditorItemID();
+										if (CurveEditorTreeItem != FCurveEditorTreeItemID::Invalid())
+										{
+											CurveEditorTreeView->SetItemSelection(CurveEditorTreeItem, bShowCurve);
+											break;
+										}
+									}
+								}
+							}
+						}
+					}
+				}
+			}
+		}
+	}
+}
+
+bool USequencerCurveEditorObject::IsCurveShown(const FSequencerChannelProxy& ChannelProxy)
+{
+	TOptional<FCurveModelID> CurveModelID;
+	if (UMovieSceneSection* Section = ChannelProxy.Section)
+	{
+		CurveModelID = USequencerCurveEditorObject::GetCurve(Section, ChannelProxy.ChannelName);
+	}
+	return CurveModelID.IsSet();
+}
+
+TOptional<FCurveModelID> USequencerCurveEditorObject::GetCurve(UMovieSceneSection* InSection, const FName& InName)
+{
+	TOptional<FCurveModelID> OptCurveModel;
 	TSharedPtr<FCurveEditor> CurveEditor = GetCurveEditor();
 	if (CurrentSequencer.IsValid() && CurveEditor.IsValid())
 	{
@@ -147,11 +212,12 @@ FCurveModelID USequencerCurveEditorObject::GetCurve(UMovieSceneSection* InSectio
 		{
 			if (Pair.Value.IsValid() && Pair.Value->GetOwningObject() == InSection && Pair.Value->GetChannelName() == InName)
 			{
-				return Pair.Key;
+				OptCurveModel = Pair.Key;
+				break;
 			}
 		}
 	}
-	return FCurveModelID::Unique();
+	return OptCurveModel;
 }
 
 void USequencerCurveEditorObject::SelectKeys(const FSequencerChannelProxy& ChannelProxy, const TArray<int32>& Indices)
@@ -161,7 +227,8 @@ void USequencerCurveEditorObject::SelectKeys(const FSequencerChannelProxy& Chann
 	{
 		if (UMovieSceneSection* Section = ChannelProxy.Section)
 		{
-			FCurveModelID CurveModelID = USequencerCurveEditorObject::GetCurve(Section, ChannelProxy.ChannelName);
+			TOptional<FCurveModelID> CurveModelID = USequencerCurveEditorObject::GetCurve(Section, ChannelProxy.ChannelName);
+			if(CurveModelID.IsSet())
 			{
 				if (FMovieSceneChannel* MovieSceneChannel = UMovieSceneSectionExtensions::GetMovieSceneChannel(Section, ChannelProxy.ChannelName))
 				{
@@ -174,7 +241,7 @@ void USequencerCurveEditorObject::SelectKeys(const FSequencerChannelProxy& Chann
 							Handles.Add(Handle);
 						}
 					}
-					CurveEditor->Selection.Add(CurveModelID, ECurvePointType::Key, Handles);
+					CurveEditor->Selection.Add(CurveModelID.GetValue(), ECurvePointType::Key, Handles);
 				}
 			}
 		}
@@ -263,6 +330,28 @@ void USequencerCurveEditorObject::SetRandomColorForChannels(UClass* Class, const
 	}
 }
 
+void USequencerCurveEditorObject::ApplyFilter(UCurveEditorFilterBase* Filter)
+{
+	TSharedPtr<FCurveEditor> CurveEditor =  GetCurveEditor();
+
+	if (Filter && CurveEditor.IsValid())
+	{
+		Filter->InitializeFilter(CurveEditor.ToSharedRef());
+		const TMap<FCurveModelID, FKeyHandleSet>& SelectedKeys = CurveEditor->GetSelection().GetAll();
+		TMap<FCurveModelID, FKeyHandleSet> OutKeysToSelect;
+		Filter->ApplyFilter(CurveEditor.ToSharedRef(), SelectedKeys, OutKeysToSelect);
+
+		// Clear their selection and then set it to the keys the filter thinks you should have selected.
+		CurveEditor->GetSelection().Clear();
+
+		for (const TTuple<FCurveModelID, FKeyHandleSet>& OutSet : OutKeysToSelect)
+		{
+			CurveEditor->GetSelection().Add(OutSet.Key, ECurvePointType::Key, OutSet.Value.AsArray());
+		}
+	}
+}
+
+
 
 
 
diff --git a/Engine/Plugins/MovieScene/SequencerScripting/Source/SequencerScriptingEditor/Public/SequencerCurveEditorObject.h b/Engine/Plugins/MovieScene/SequencerScripting/Source/SequencerScriptingEditor/Public/SequencerCurveEditorObject.h
--- a/Engine/Plugins/MovieScene/SequencerScripting/Source/SequencerScriptingEditor/Public/SequencerCurveEditorObject.h
+++ b/Engine/Plugins/MovieScene/SequencerScripting/Source/SequencerScriptingEditor/Public/SequencerCurveEditorObject.h
@@ -12,7 +12,7 @@ class FCurveEditor;
 struct FMovieSceneChannel;
 class UMovieSceneSection;
 struct FCurveModelID;
-
+class UCurveEditorFilterBase;
 
 USTRUCT(BlueprintType)
 struct SEQUENCERSCRIPTINGEDITOR_API FSequencerChannelProxy
@@ -59,6 +59,9 @@ public:
 	/** Close curve editor*/
 	void CloseCurveEditor();
 
+	UFUNCTION(BlueprintCallable, Category = "Sequencer Curve Editor")
+	void ApplyFilter(UCurveEditorFilterBase* Filter);
+
 public:
 
 	/** Gets the channel with selected keys */
@@ -69,6 +72,14 @@ public:
 	UFUNCTION(BlueprintPure, Category = "Sequencer Curve Editor")
 	TArray<int32> GetSelectedKeys(const FSequencerChannelProxy& ChannelProxy);
 
+	/** Show curve  */
+	UFUNCTION(BlueprintCallable, Category = "Sequencer Curve Editor")
+	void ShowCurve(const FSequencerChannelProxy& Channel, bool bShowCurve);
+
+	/** Is the curve displayed*/
+	UFUNCTION(BlueprintCallable, Category = "Sequencer Curve Editor")
+	bool IsCurveShown(const FSequencerChannelProxy& Channel);
+
 	/** Select keys */
 	UFUNCTION(BlueprintCallable, Category = "Sequencer Curve Editor")
 	void SelectKeys(const FSequencerChannelProxy& Channel, const TArray<int32>& Indices);
@@ -114,9 +125,9 @@ public:
 public:
 
 	/**
-	 * Utility function to get curve from a section and a name
+	 * Utility function to get curve, if it exists,  from a section and a name
 	 */
-	FCurveModelID GetCurve(UMovieSceneSection* InSection, const FName& InName);
+	TOptional<FCurveModelID> GetCurve(UMovieSceneSection* InSection, const FName& InName);
 
 	/**
 	Utility function to get curve editor
diff --git a/Engine/Source/Editor/CurveEditor/Private/Filters/CurveEditorGaussianFilter.cpp b/Engine/Source/Editor/CurveEditor/Private/Filters/CurveEditorGaussianFilter.cpp
new file mode 100644
--- /dev/null
+++ b/Engine/Source/Editor/CurveEditor/Private/Filters/CurveEditorGaussianFilter.cpp
@@ -0,0 +1,140 @@
+// Copyright Epic Games, Inc. All Rights Reserved.
+
+#include "Filters/CurveEditorGaussianFilter.h"
+
+#include "Containers/Array.h"
+#include "Containers/ArrayView.h"
+#include "CurveDataAbstraction.h"
+#include "CurveEditor.h"
+#include "CurveEditorSelection.h"
+#include "CurveEditorSnapMetrics.h"
+#include "CurveEditorTypes.h"
+#include "CurveModel.h"
+#include "Curves/CurveEvaluation.h"
+#include "Curves/KeyHandle.h"
+#include "Curves/RealCurve.h"
+#include "Curves/RichCurve.h"
+#include "HAL/PlatformCrt.h"
+#include "Math/NumericLimits.h"
+#include "Math/UnrealMathSSE.h"
+#include "Misc/FrameNumber.h"
+#include "Misc/FrameRate.h"
+#include "Misc/FrameTime.h"
+#include "Templates/SharedPointer.h"
+#include "Templates/Tuple.h"
+#include "Templates/UnrealTemplate.h"
+
+
+TArray<double> UCurveEditorGaussianFilter::CalculateKernel(int32 InKernelWidth)
+{
+	const int32 KernelWidth = ((InKernelWidth / 2) * 2) + 1; //make sure it's odd
+	double Sigma = ((double)(KernelWidth) - 1.0) / 6.0; //99% is 6sigma
+	TArray<double> Kernel;
+	Kernel.SetNum(KernelWidth);
+	int32 Mid = (KernelWidth / 2);
+	int32 Index = 0;
+	for (int32 KernelIndex = -Mid; KernelIndex < Mid + 1; ++KernelIndex)
+	{
+		Kernel[Index] = (1.0 / (Sigma * FMath::Sqrt(2.0 * PI))) * (1.0 / (FMath::Exp((KernelIndex * KernelIndex) / (2.0 * Sigma * Sigma))));
+		++Index;
+	}
+	return Kernel;
+}
+
+double UCurveEditorGaussianFilter::ApplyKernel(int32 PositionIndex,const TArray<double> Kernel, const TArray<FKeyPosition>& KeyPositions)
+{
+	const int32 StartKernel = - Kernel.Num() / 2;
+	const int32 EndKernel = StartKernel + Kernel.Num();
+	double TotalWeight = 0.0;
+	double Sum = 0.0;
+	int32 KernelIndex = 0;
+	for (int32 KernelLocation = StartKernel; KernelLocation < EndKernel; ++KernelLocation)
+	{
+		int32 Index = KernelLocation + PositionIndex;
+		if (Index < 0)
+		{
+			Index = 0;
+		}
+		else if (Index >= KeyPositions.Num())
+		{
+			Index = KeyPositions.Num() - 1;
+		}
+		TotalWeight += Kernel[KernelIndex];
+		Sum += (Kernel[KernelIndex] * KeyPositions[Index].OutputValue);
+
+		++KernelIndex;
+	}
+	if (TotalWeight < 0.9999 && TotalWeight > 0.0)
+	{
+		Sum *= (1.0 / TotalWeight);
+	}
+	return Sum;
+}
+
+void UCurveEditorGaussianFilter::Gaussian(FCurveModel* Curve, const FGaussianParams& InParams, const TOptional<FKeyHandleSet>& KeyHandleSet,
+	FKeyHandleSet& OutHandleSet)
+{
+	TArray<FKeyHandle> KeyHandles;
+	TArray<FKeyPosition> SelectedKeyPositions;
+	double MinKey = TNumericLimits<double>::Max(), MaxKey = TNumericLimits<double>::Lowest();
+
+	//if key's are set use that to find range to reduce otherwise use full range.
+	if (KeyHandleSet.IsSet())
+	{
+		KeyHandles.Reset(KeyHandleSet.GetValue().Num());
+		KeyHandles.Append(KeyHandleSet.GetValue().AsArray().GetData(), KeyHandleSet.GetValue().Num());
+		//assume in order todo check
+		SelectedKeyPositions.SetNum(KeyHandles.Num());
+		Curve->GetKeyPositions(KeyHandles, SelectedKeyPositions);
+		// Find the hull of the range of the selected keys
+		for (FKeyPosition Key : SelectedKeyPositions)
+		{
+			MinKey = FMath::Min(Key.InputValue, MinKey);
+			MaxKey = FMath::Max(Key.InputValue, MaxKey);
+		}
+	}
+	else //get all keys
+	{
+		Curve->GetTimeRange(MinKey, MaxKey);
+	}
+
+	// Get all keys that exist between the time range
+	KeyHandles.Reset();
+	Curve->GetKeys(MinKey, MaxKey, TNumericLimits<double>::Lowest(), TNumericLimits<double>::Max(), KeyHandles);
+	if (KeyHandles.Num() <= 2)
+	{
+		return;
+	}
+	
+	TArray<FKeyPosition> KeyPositions, NewKeyPositions;
+	KeyPositions.SetNum(KeyHandles.Num());
+	Curve->GetKeyPositions(KeyHandles, KeyPositions);
+	NewKeyPositions = KeyPositions;
+	TArray<double> Kernel = CalculateKernel(InParams.KernelWidth);
+	for (int32 Index = 0; Index < KeyPositions.Num(); ++Index)
+	{
+		NewKeyPositions[Index].OutputValue = ApplyKernel(Index, Kernel, KeyPositions);
+	}
+	Curve->SetKeyPositions(KeyHandles, NewKeyPositions);
+	for (const FKeyHandle& NewHandle : KeyHandles)
+	{
+		OutHandleSet.Add(NewHandle, ECurvePointType::Key);
+	}
+}
+
+void UCurveEditorGaussianFilter::ApplyFilter_Impl(TSharedRef<FCurveEditor> InCurveEditor, const TMap<FCurveModelID, FKeyHandleSet>& InKeysToOperateOn, TMap<FCurveModelID, FKeyHandleSet>& OutKeysToSelect)
+{
+	OutKeysToSelect.Reset();
+	for (const TTuple<FCurveModelID, FKeyHandleSet>& Pair : InKeysToOperateOn)
+	{
+		FCurveModel* Curve = InCurveEditor->FindCurve(Pair.Key);
+		if (!Curve)
+		{
+			continue;
+		}
+		FKeyHandleSet& OutHandleSet = OutKeysToSelect.Add(Pair.Key);
+
+		TOptional<FKeyHandleSet> KeyHandleSet = Pair.Value;
+		Gaussian(Curve, GaussianParams, KeyHandleSet, OutHandleSet);
+	}
+}
\ No newline at end of file
diff --git a/Engine/Source/Editor/CurveEditor/Private/Filters/CurveEditorSmartReduceFilter.cpp b/Engine/Source/Editor/CurveEditor/Private/Filters/CurveEditorSmartReduceFilter.cpp
--- a/Engine/Source/Editor/CurveEditor/Private/Filters/CurveEditorSmartReduceFilter.cpp
+++ b/Engine/Source/Editor/CurveEditor/Private/Filters/CurveEditorSmartReduceFilter.cpp
@@ -290,7 +290,7 @@ void UCurveEditorSmartReduceFilter::SmartReduce(FCurveModel* Curve, const FSmart
 	}
 	const float Tolerance = (InParams.TolerancePercentage / 100.0) * (Cache.Max - Cache.Min);
 	TArray<FKeyHandle> HandlesAdded;
-	for (int32 Index = 0; Index < (NewPositions.Num() -2); ++Index)
+	for (int32 Index = 0; Index < (NewPositions.Num() -1); ++Index)
 	{
 		const int32 StartIndex = Cache.GetIndex(NewPositions[Index].InputValue);
 		const int32 EndIndex = Cache.GetIndex(NewPositions[Index + 1].InputValue);
diff --git a/Engine/Source/Editor/CurveEditor/Public/Filters/CurveEditorFilterBase.h b/Engine/Source/Editor/CurveEditor/Public/Filters/CurveEditorFilterBase.h
--- a/Engine/Source/Editor/CurveEditor/Public/Filters/CurveEditorFilterBase.h
+++ b/Engine/Source/Editor/CurveEditor/Public/Filters/CurveEditorFilterBase.h
@@ -20,7 +20,7 @@ class FCurveEditor;
 * The CDO is directly edited in the Curve Editor Filter panel so if you wish to save user settings between runs you can
 * add the "config" property to your UPROPERTY() and it will automatically get saved.
 */
-UCLASS(config=EditorSettings)
+UCLASS(BlueprintType, config=EditorSettings)
 class CURVEEDITOR_API UCurveEditorFilterBase : public UObject
 {
 	GENERATED_BODY()
diff --git a/Engine/Source/Editor/CurveEditor/Public/Filters/CurveEditorGaussianFilter.h b/Engine/Source/Editor/CurveEditor/Public/Filters/CurveEditorGaussianFilter.h
new file mode 100644
--- /dev/null
+++ b/Engine/Source/Editor/CurveEditor/Public/Filters/CurveEditorGaussianFilter.h
@@ -0,0 +1,59 @@
+// Copyright Epic Games, Inc. All Rights Reserved.
+
+#pragma once
+
+#include "Containers/Map.h"
+#include "Filters/CurveEditorFilterBase.h"
+#include "Misc/FrameRate.h"
+#include "Templates/SharedPointer.h"
+#include "UObject/ObjectMacros.h"
+#include "UObject/UObjectGlobals.h"
+
+#include "CurveEditorGaussianFilter.generated.h"
+
+class FCurveEditor;
+class UObject;
+struct FCurveModelID;
+struct FKeyHandleSet;
+
+
+USTRUCT(BlueprintType)
+struct FGaussianParams
+{
+	GENERATED_BODY()
+	/** Sigma that determines how 'fat' the filter is, higher value means fatter width  */
+	UPROPERTY(EditAnywhere, BlueprintReadWrite, Category = "Settings", meta = (ClampMin = "3", ClampMax = "9999", UIMin = "3", UIMax = "9999"))
+	int32 KernelWidth = 5;
+
+	void Reset()
+	{
+		KernelWidth = 5;
+	}
+};
+
+UCLASS(DisplayName = "Gaussian", CollapseCategories)
+class UCurveEditorGaussianFilter : public UCurveEditorFilterBase
+{
+	GENERATED_BODY()
+public:
+	UCurveEditorGaussianFilter() {};
+	/**
+	* @param Curve to perform Gaussian
+	* @param InParams Params to use to peform reduction
+	* @param KeyHandleSet Optional set of keys to operate on, if not set we will operate on all of them
+	* @param OutHandleSet The set of new keys.
+	*/
+	static CURVEEDITOR_API void Gaussian(FCurveModel* Curve, const FGaussianParams& InParams, const TOptional<FKeyHandleSet>& KeyHandleSet,
+		FKeyHandleSet& OutHandleSet);
+protected:
+	void ApplyFilter_Impl(TSharedRef<FCurveEditor> InCurveEditor, const TMap<FCurveModelID, FKeyHandleSet>& InKeysToOperateOn, TMap<FCurveModelID, FKeyHandleSet>& OutKeysToSelect) override;
+
+	static TArray<double>  CalculateKernel(int32 KernelWidth);
+	static double ApplyKernel(int32 Index, const TArray<double> Kernel, const TArray<FKeyPosition>& KeyPositions);
+
+public:
+
+	UPROPERTY(EditAnywhere, BlueprintReadWrite, Category = "Settings")
+	FGaussianParams GaussianParams;
+
+};
\ No newline at end of file
