Sequencer: Remove deprecations (5.1, 5.2)

Removed a lot of the deprecated notions of 'master' and also removed the exceptions in CheckUnacceptableWords!

#rnx

diff --git a/Engine/Plugins/MovieScene/MovieRenderPipeline/Source/MovieRenderPipelineCore/Private/MoviePipeline.cpp b/Engine/Plugins/MovieScene/MovieRenderPipeline/Source/MovieRenderPipelineCore/Private/MoviePipeline.cpp
--- a/Engine/Plugins/MovieScene/MovieRenderPipeline/Source/MovieRenderPipelineCore/Private/MoviePipeline.cpp
+++ b/Engine/Plugins/MovieScene/MovieRenderPipeline/Source/MovieRenderPipelineCore/Private/MoviePipeline.cpp
@@ -1628,12 +1628,6 @@ FMoviePipelineOutputData UMoviePipeline::GetOutputDataParams()
 
 void UMoviePipeline::OnMoviePipelineFinishedImpl()
 {
-	// Broadcast to both Native and Python/BP
-	PRAGMA_DISABLE_DEPRECATION_WARNINGS
-	OnMoviePipelineFinishedDelegateNative.Broadcast(this, bFatalError);
-	OnMoviePipelineFinishedDelegate.Broadcast(this, bFatalError);
-	PRAGMA_ENABLE_DEPRECATION_WARNINGS
-
 	// Generate a params struct containing the data generated by this job.
 	FMoviePipelineOutputData Params;
 	Params.Pipeline = this;
diff --git a/Engine/Plugins/MovieScene/MovieRenderPipeline/Source/MovieRenderPipelineCore/Private/MoviePipelineSettingBlueprintBase.cpp b/Engine/Plugins/MovieScene/MovieRenderPipeline/Source/MovieRenderPipelineCore/Private/MoviePipelineSettingBlueprintBase.cpp
--- a/Engine/Plugins/MovieScene/MovieRenderPipeline/Source/MovieRenderPipelineCore/Private/MoviePipelineSettingBlueprintBase.cpp
+++ b/Engine/Plugins/MovieScene/MovieRenderPipeline/Source/MovieRenderPipelineCore/Private/MoviePipelineSettingBlueprintBase.cpp
@@ -4,11 +4,4 @@
 void UMoviePipelineSetting_BlueprintBase::PostLoad()
 {
 	Super::PostLoad();
-
-#if WITH_EDITOR
-	if (!bIsValidOnMaster_DEPRECATED)
-	{
-		bIsValidOnPrimary = bIsValidOnMaster_DEPRECATED;
-	}
-#endif
 }
diff --git a/Engine/Plugins/MovieScene/MovieRenderPipeline/Source/MovieRenderPipelineCore/Public/MoviePipeline.h b/Engine/Plugins/MovieScene/MovieRenderPipeline/Source/MovieRenderPipelineCore/Public/MoviePipeline.h
--- a/Engine/Plugins/MovieScene/MovieRenderPipeline/Source/MovieRenderPipelineCore/Public/MoviePipeline.h
+++ b/Engine/Plugins/MovieScene/MovieRenderPipeline/Source/MovieRenderPipelineCore/Public/MoviePipeline.h
@@ -88,18 +88,6 @@ public:
 	UFUNCTION(BlueprintCallable, Category = "Movie Render Pipeline")
 	void SetInitializationTime(const FDateTime& InDateTime) { InitializationTime = InDateTime; InitializationTimeOffset = FTimespan(); }
 
-	/** Deprecated. Use OnMoviePipelineWorkFinished() instead. */
-	UE_DEPRECATED(4.27, "Use OnMoviePipelineWorkFinished() instead.")
-	FMoviePipelineFinishedNative& OnMoviePipelineFinished()
-	{
-		return OnMoviePipelineFinishedDelegateNative;
-	}
-
-	/** Deprecated. Use OnMoviePipelineWorkFinishedDelegate instead */
-	UE_DEPRECATED(4.27, "Use OnMoviePipelineWorkFinishedDelegate instead.")
-	UPROPERTY(BlueprintAssignable, Category = "Movie Render Pipeline")
-	FMoviePipelineFinished OnMoviePipelineFinishedDelegate;
-
 	/**
 	* Get the Primary Configuration used to render this shot. This contains the global settings for the shot, as well as per-shot
 	* configurations which can contain their own settings.
@@ -107,10 +95,6 @@ public:
 	UFUNCTION(BlueprintPure, Category = "Movie Render Pipeline")
 	UMoviePipelinePrimaryConfig* GetPipelinePrimaryConfig() const;
 	
-	UE_DEPRECATED(5.2, "GetPipelineMasterConfig is deprecated. Please use GetPipelinePrimaryConfig instead")
-	UFUNCTION(BlueprintPure, Category = "Movie Render Pipeline", meta=(DeprecatedFunction, DeprecationMessage = "Use GetPipelinePrimaryConfig"))
-	UMoviePipelinePrimaryConfig* GetPipelineMasterConfig() const { return GetPipelinePrimaryConfig(); }
-
 public:
 	ULevelSequence* GetTargetSequence() const { return TargetSequence; }
 
@@ -395,9 +379,6 @@ private:
 	/** When we originally initialize we store the offset from UTC (which is what GetInitializationTime() is in), but we clear this if you call SetInitializationTime. */
 	FTimespan InitializationTimeOffset;
 
-	/** Deprecated. */
-	FMoviePipelineFinishedNative OnMoviePipelineFinishedDelegateNative;
-
 	/**
 	 * We have to apply camera motion vectors manually. So we keep the current and previous frame's camera view and rotation.
 	 * Then we render a sequence of the same movement, and update after running the game sim.
@@ -412,11 +393,6 @@ public:
 	/** A debug image sequence writer in the event they want to dump every sample generated on its own. */
 	IImageWriteQueue* ImageWriteQueue;
 
-	/** Optional widget for feedback during render */
-	UE_DEPRECATED(5.1, "Use SetViewportInitArgs instead.")
-	UPROPERTY(Transient)
-	TSubclassOf<UMovieRenderDebugWidget> DebugWidgetClass;
-	
 	/** Used to track first-render submissions (for 3d renders) to set the correct flags on the renderer module. */
 	bool bHasRenderedFirstViewThisFrame;
 
diff --git a/Engine/Plugins/MovieScene/MovieRenderPipeline/Source/MovieRenderPipelineCore/Public/MoviePipelineBlueprintLibrary.h b/Engine/Plugins/MovieScene/MovieRenderPipeline/Source/MovieRenderPipelineCore/Public/MoviePipelineBlueprintLibrary.h
--- a/Engine/Plugins/MovieScene/MovieRenderPipeline/Source/MovieRenderPipelineCore/Public/MoviePipelineBlueprintLibrary.h
+++ b/Engine/Plugins/MovieScene/MovieRenderPipeline/Source/MovieRenderPipelineCore/Public/MoviePipelineBlueprintLibrary.h
@@ -104,17 +104,9 @@ public:
 	UFUNCTION(BlueprintPure, Category = "Movie Render Pipeline")
 	static FTimecode GetRootTimecode(const UMoviePipeline* InMoviePipeline);
 
-	UE_DEPRECATED(5.2, "GetMasterTimecode is deprecated. Please use GetRootTimecode instead")
-	UFUNCTION(BlueprintPure, Category = "Movie Render Pipeline", meta = (DeprecatedFunction, DeprecationMessage = "GetMasterTimecode is deprecated. Please use GetRootTimecode instead"))
-	static FTimecode GetMasterTimecode(const UMoviePipeline* InMoviePipeline) { return GetRootTimecode(InMoviePipeline); }
-
 	UFUNCTION(BlueprintPure, Category = "Movie Render Pipeline")
 	static FFrameNumber GetRootFrameNumber(const UMoviePipeline* InMoviePipeline);
 
-	UE_DEPRECATED(5.2, "GetMasterFrameNumber is deprecated. Please use GetRootFrameNumber instead")
-	UFUNCTION(BlueprintPure, Category = "Movie Render Pipeline", meta = (DeprecatedFunction, DeprecationMessage = "GetMasterFrameNumber is deprecated. Please GetRootFrameNumber instead"))
-	static FFrameNumber GetMasterFrameNumber(const UMoviePipeline* InMoviePipeline) { return GetRootFrameNumber(InMoviePipeline); }
-
 	UFUNCTION(BlueprintPure, Category = "Movie Render Pipeline")
 	static FTimecode GetCurrentShotTimecode(const UMoviePipeline* InMoviePipeline);
 
diff --git a/Engine/Plugins/MovieScene/MovieRenderPipeline/Source/MovieRenderPipelineCore/Public/MoviePipelineMasterConfig.h b/Engine/Plugins/MovieScene/MovieRenderPipeline/Source/MovieRenderPipelineCore/Public/MoviePipelineMasterConfig.h
deleted file mode 100644
diff --git a/Engine/Plugins/MovieScene/MovieRenderPipeline/Source/MovieRenderPipelineCore/Public/MoviePipelineSetting.h b/Engine/Plugins/MovieScene/MovieRenderPipeline/Source/MovieRenderPipelineCore/Public/MoviePipelineSetting.h
--- a/Engine/Plugins/MovieScene/MovieRenderPipeline/Source/MovieRenderPipelineCore/Public/MoviePipelineSetting.h
+++ b/Engine/Plugins/MovieScene/MovieRenderPipeline/Source/MovieRenderPipelineCore/Public/MoviePipelineSetting.h
@@ -60,15 +60,6 @@ public:
 	* allows a setting to provide these when the user wants to run in a separate process. This won't
 	* be used when running in the current process because it is too late to modify the command line.
 	*/
-	UE_DEPRECATED(5.0, "Use BuildNewProcessCommandLineArgs instead.")
-	UFUNCTION(BlueprintCallable, Category = "Movie Render Pipeline")
-	void BuildNewProcessCommandLine(UPARAM(ref) FString& InOutUnrealURLParams, UPARAM(ref) FString& InOutCommandLineArgs) const
-	{
-		PRAGMA_DISABLE_DEPRECATION_WARNINGS
-		BuildNewProcessCommandLineImpl(InOutUnrealURLParams, InOutCommandLineArgs); 
-		PRAGMA_ENABLE_DEPRECATION_WARNINGS
-	}
-
 	UFUNCTION(BlueprintCallable, Category = "Movie Render Pipeline")
 	void BuildNewProcessCommandLineArgs(UPARAM(ref) TArray<FString>& InOutUnrealURLParams, UPARAM(ref) TArray<FString>& InOutCommandLineArgs, UPARAM(ref) TArray<FString>& InOutDeviceProfileCvars, UPARAM(ref) TArray<FString>& InOutExecCmds) const { BuildNewProcessCommandLineArgsImpl(InOutUnrealURLParams, InOutCommandLineArgs, InOutDeviceProfileCvars, InOutExecCmds); }
 
@@ -120,9 +111,6 @@ public:
 	/** Can this configuration setting be added to the primary configuration? If not, it will throw an error when trying to add it to the primary configuration. */
 	virtual bool IsValidOnPrimary() const PURE_VIRTUAL(UMoviePipelineSetting::IsValidOnPrimary, return false; );
 
-	UE_DEPRECATED(5.2, "IsValidOnMaster is deprecated. Please use IsValidOnPrimary instead")
-	virtual bool IsValidOnMaster() const { return IsValidOnPrimary(); }
-
 	/**
 	* If true, then this setting will be included when searching for settings even if it was added transiently. This is used for the rare case where a setting
 	* needs to be run (to set reasonable default values) even if the user hasn't added it.
@@ -147,8 +135,6 @@ public:
 	/** Return Key/Value pairs that you wish to be usable in the Output File Name format string or file metadata. This allows settings to add format strings based on their values. */
 	virtual void GetFormatArguments(FMoviePipelineFormatArgs& InOutFormatArgs) const {}
 	
-	UE_DEPRECATED(5.0, "Use BuildNewProcessCommandLineArgsImpl instead.")
-	virtual void BuildNewProcessCommandLineImpl(FString& InOutUnrealURLParams, FString& InOutCommandLineArgs) const { }
 	/** Modify the Unreal URL and Command Line Arguments when preparing the setting to be run in a new process. */
 	virtual void BuildNewProcessCommandLineArgsImpl(TArray<FString>& InOutUnrealURLParams, TArray<FString>& InOutCommandLineArgs, TArray<FString>& InOutDeviceProfileCvars, TArray<FString>& InOutExecCmds) const { }
 
diff --git a/Engine/Plugins/MovieScene/MovieRenderPipeline/Source/MovieRenderPipelineCore/Public/MoviePipelineSettingBlueprintBase.h b/Engine/Plugins/MovieScene/MovieRenderPipeline/Source/MovieRenderPipelineCore/Public/MoviePipelineSettingBlueprintBase.h
--- a/Engine/Plugins/MovieScene/MovieRenderPipeline/Source/MovieRenderPipelineCore/Public/MoviePipelineSettingBlueprintBase.h
+++ b/Engine/Plugins/MovieScene/MovieRenderPipeline/Source/MovieRenderPipelineCore/Public/MoviePipelineSettingBlueprintBase.h
@@ -22,9 +22,6 @@ public:
 	{
 		CategoryText = NSLOCTEXT("MovieRenderPipelineBP", "DefaultCategoryName_Text", "Custom Settings");
 		bIsValidOnPrimary = true;
-#if WITH_EDITORONLY_DATA
-		bIsValidOnMaster_DEPRECATED = true;
-#endif
 		bIsValidOnShots = true;
 		bCanBeDisabled = true;
 	}
@@ -89,11 +86,4 @@ protected:
 	
 	UPROPERTY(EditDefaultsOnly, Category = "Movie Pipeline")
 	bool bCanBeDisabled;
-
-private:
-
-#if WITH_EDITORONLY_DATA
-	UPROPERTY()
-	bool bIsValidOnMaster_DEPRECATED;
-#endif
 };
\ No newline at end of file
diff --git a/Engine/Plugins/MovieScene/MovieRenderPipeline/Source/MovieRenderPipelineEditor/Private/Widgets/SMoviePipelineConfigPanel.cpp b/Engine/Plugins/MovieScene/MovieRenderPipeline/Source/MovieRenderPipelineEditor/Private/Widgets/SMoviePipelineConfigPanel.cpp
--- a/Engine/Plugins/MovieScene/MovieRenderPipeline/Source/MovieRenderPipelineEditor/Private/Widgets/SMoviePipelineConfigPanel.cpp
+++ b/Engine/Plugins/MovieScene/MovieRenderPipeline/Source/MovieRenderPipelineEditor/Private/Widgets/SMoviePipelineConfigPanel.cpp
@@ -418,15 +418,6 @@ TSharedRef<SWidget> SMoviePipelineConfigPanel::OnGeneratePresetsMenu()
 
 		AssetPickerConfig.AssetShowWarningText = LOCTEXT("NoPresets_Warning", "No Presets Found");
 		AssetPickerConfig.Filter.ClassPaths.Add(ConfigAssetType->GetClassPathName());
-
-		// The redirected masterconfig assets aren't found when the new primaryconfig asset path is specified. This is a temporary workaround 
-		// for the AssetRegistry not applying CoreRedirects to the ClassPath. 
-		// This can be removed once that bug (UE-168245) is fixed and the KnownCase in CheckUnacceptableWords can also be removed.
-		if (ConfigAssetType->GetClassPathName() == FTopLevelAssetPath(TEXT("/Script/MovieRenderPipelineCore.MoviePipelinePrimaryConfig")))
-		{
-			AssetPickerConfig.Filter.ClassPaths.Add(FTopLevelAssetPath(TEXT("/Script/MovieRenderPipelineCore.MoviePipelineMasterConfig")));
-		}
-		
 		AssetPickerConfig.Filter.bRecursiveClasses = true;
 		AssetPickerConfig.OnAssetSelected = FOnAssetSelected::CreateSP(this, &SMoviePipelineConfigPanel::OnImportPreset);
 	}
diff --git a/Engine/Plugins/MovieScene/MovieRenderPipeline/Source/MovieRenderPipelineEditor/Private/Widgets/SMoviePipelineQueueEditor.cpp b/Engine/Plugins/MovieScene/MovieRenderPipeline/Source/MovieRenderPipelineEditor/Private/Widgets/SMoviePipelineQueueEditor.cpp
--- a/Engine/Plugins/MovieScene/MovieRenderPipeline/Source/MovieRenderPipelineEditor/Private/Widgets/SMoviePipelineQueueEditor.cpp
+++ b/Engine/Plugins/MovieScene/MovieRenderPipeline/Source/MovieRenderPipelineEditor/Private/Widgets/SMoviePipelineQueueEditor.cpp
@@ -487,15 +487,6 @@ public:
 
 			AssetPickerConfig.AssetShowWarningText = NoAssetsFoundWarning;
 			AssetPickerConfig.Filter.ClassPaths.Add(InClass->GetClassPathName());
-
-			// The redirected masterconfig assets aren't found when the new primaryconfig asset path is specified. This is a temporary workaround 
-			// for the AssetRegistry not applying CoreRedirects to the ClassPath. 
-			// This can be removed once that bug (UE-168245) is fixed and the KnownCase in CheckUnacceptableWords can also be removed.
-			if (InClass->GetClassPathName() == FTopLevelAssetPath(TEXT("/Script/MovieRenderPipelineCore.MoviePipelinePrimaryConfig")))
-			{
-				AssetPickerConfig.Filter.ClassPaths.Add(FTopLevelAssetPath(TEXT("/Script/MovieRenderPipelineCore.MoviePipelineMasterConfig")));
-			}
-
 			AssetPickerConfig.Filter.bRecursiveClasses = true;
 			AssetPickerConfig.OnAssetSelected = InOnAssetSelected;
 		}
diff --git a/Engine/Plugins/MovieScene/MovieRenderPipeline/Source/MovieRenderPipelineEditor/Public/MoviePipelineConfigFactory.h b/Engine/Plugins/MovieScene/MovieRenderPipeline/Source/MovieRenderPipelineEditor/Public/MoviePipelineConfigFactory.h
--- a/Engine/Plugins/MovieScene/MovieRenderPipeline/Source/MovieRenderPipelineEditor/Public/MoviePipelineConfigFactory.h
+++ b/Engine/Plugins/MovieScene/MovieRenderPipeline/Source/MovieRenderPipelineEditor/Public/MoviePipelineConfigFactory.h
@@ -5,10 +5,6 @@
 #include "Factories/Factory.h"
 #include "MoviePipelineConfigFactory.generated.h"
 
-// For backwards compatibility 5.2
-UE_DEPRECATED(5.2, "UMoviePipelineMasterConfigFactory is deprecated. Please use UMoviePipelinePrimaryConfigFactory.")
-typedef UMoviePipelinePrimaryConfigFactory UMoviePipelineMasterConfigFactory;
-
 // Forward Declare
 UCLASS(BlueprintType)
 class MOVIERENDERPIPELINEEDITOR_API UMoviePipelinePrimaryConfigFactory : public UFactory
diff --git a/Engine/Plugins/MovieScene/MovieRenderPipeline/Source/MovieRenderPipelineRenderPasses/Public/MoviePipelineImagePassBase.h b/Engine/Plugins/MovieScene/MovieRenderPipeline/Source/MovieRenderPipelineRenderPasses/Public/MoviePipelineImagePassBase.h
--- a/Engine/Plugins/MovieScene/MovieRenderPipeline/Source/MovieRenderPipelineRenderPasses/Public/MoviePipelineImagePassBase.h
+++ b/Engine/Plugins/MovieScene/MovieRenderPipeline/Source/MovieRenderPipelineRenderPasses/Public/MoviePipelineImagePassBase.h
@@ -120,8 +120,6 @@ protected:
 protected:
 	virtual void GetViewShowFlags(FEngineShowFlags& OutShowFlag, EViewModeIndex& OutViewModeIndex) const;
 	virtual TSharedPtr<FSceneViewFamilyContext> CalculateViewFamily(FMoviePipelineRenderPassMetrics& InOutSampleState, IViewCalcPayload* OptPayload = nullptr);
-	UE_DEPRECATED(5.1, "Use the overload that provides the InOutSampleState and IViewCalcPayload instead.")
-	virtual void BlendPostProcessSettings(FSceneView* InView) {}
 	virtual void BlendPostProcessSettings(FSceneView* InView, FMoviePipelineRenderPassMetrics& InOutSampleState, IViewCalcPayload* OptPayload = nullptr);
 	virtual void SetupViewForViewModeOverride(FSceneView* View);
 	virtual void MoviePipelineRenderShowFlagOverride(FEngineShowFlags& OutShowFlag) {}
diff --git a/Engine/Plugins/MovieScene/MovieRenderPipeline/Source/MovieRenderPipelineSettings/Private/MoviePipelineConsoleVariableSetting.cpp b/Engine/Plugins/MovieScene/MovieRenderPipeline/Source/MovieRenderPipelineSettings/Private/MoviePipelineConsoleVariableSetting.cpp
--- a/Engine/Plugins/MovieScene/MovieRenderPipeline/Source/MovieRenderPipelineSettings/Private/MoviePipelineConsoleVariableSetting.cpp
+++ b/Engine/Plugins/MovieScene/MovieRenderPipeline/Source/MovieRenderPipelineSettings/Private/MoviePipelineConsoleVariableSetting.cpp
@@ -220,7 +220,6 @@ void UMoviePipelineConsoleVariableSetting::ApplyCVarSettings(const bool bOverrid
 {
 	if (bOverrideValues)
 	{
-		MergeInOldConsoleVariables();
 		MergeInPresetConsoleVariables();
 		PreviousConsoleVariableValues.Reset();
 		PreviousConsoleVariableValues.SetNumZeroed(MergedConsoleVariables.Num());
@@ -349,22 +348,6 @@ void UMoviePipelineConsoleVariableSetting::PostLoad()
 	Super::PostLoad();
 
 #if WITH_EDITOR
-	PRAGMA_DISABLE_DEPRECATION_WARNINGS
-
-	// Convert the old ConsoleVariables map into the new struct array
-	if (!ConsoleVariables_DEPRECATED.IsEmpty())
-	{
-		for (const TPair<FString, float>& Pair : ConsoleVariables_DEPRECATED)
-		{
-			// The converted cvar will be enabled by default
-			CVars.Add(FMoviePipelineConsoleVariableEntry(Pair.Key, Pair.Value));
-		}
-
-		ConsoleVariables_DEPRECATED.Empty();
-	}
-	
-	PRAGMA_ENABLE_DEPRECATION_WARNINGS
-
 	// Initialize the CommandInfo structs
 	for (FMoviePipelineConsoleVariableEntry& Entry : CVars)
 	{
@@ -405,43 +388,6 @@ void UMoviePipelineConsoleVariableSetting::PostEditChangeChainProperty(FProperty
 }
 #endif	// WITH_EDITOR
 
-void UMoviePipelineConsoleVariableSetting::MergeInOldConsoleVariables()
-{
-	// Note: The old cvars in ConsoleVariables are merged in via PostLoad(), which covers most use cases.
-	// However, scripting may have modified ConsoleVariables after initial load, so they need to be
-	// merged in again in certain scenarios, like when cvars are applied.
-	//
-	// Data loss is possible if scripting modified ConsoleVariables, but this method never executes
-	// (ie, MRQ never gets run in PIE). ConsoleVariables is deprecated, thus changes to it will not be saved.
-	
-	PRAGMA_DISABLE_DEPRECATION_WARNINGS
-	
-	for (const TPair<FString, float>& OldCVar : ConsoleVariables_DEPRECATED)
-	{
-		// Copy the value in ConsoleVariables over to CVars if an entry in CVars already exists. Apply the value to the
-		// last matching entry in CVars, since the last entry will win.
-		bool bDidFindNewCVar = false;
-		for (int32 Index = CVars.Num() - 1; Index >= 0; --Index)
-		{
-			FMoviePipelineConsoleVariableEntry& NewCVar = CVars[Index];
-			if (NewCVar.Name == OldCVar.Key)
-			{
-				NewCVar.Value = OldCVar.Value;
-				bDidFindNewCVar = true;
-				break;
-			}
-		}
-
-		// Otherwise, add a new entry to CVars
-		if (!bDidFindNewCVar)
-		{
-			CVars.Add(FMoviePipelineConsoleVariableEntry(OldCVar.Key, OldCVar.Value));
-		}
-	}
-	
-	PRAGMA_ENABLE_DEPRECATION_WARNINGS
-}
-
 void UMoviePipelineConsoleVariableSetting::MergeInPresetConsoleVariables()
 {
 	MergedConsoleVariables.Reset();
diff --git a/Engine/Plugins/MovieScene/MovieRenderPipeline/Source/MovieRenderPipelineSettings/Public/MoviePipelineConsoleVariableSetting.h b/Engine/Plugins/MovieScene/MovieRenderPipeline/Source/MovieRenderPipelineSettings/Public/MoviePipelineConsoleVariableSetting.h
--- a/Engine/Plugins/MovieScene/MovieRenderPipeline/Source/MovieRenderPipelineSettings/Public/MoviePipelineConsoleVariableSetting.h
+++ b/Engine/Plugins/MovieScene/MovieRenderPipeline/Source/MovieRenderPipelineSettings/Public/MoviePipelineConsoleVariableSetting.h
@@ -111,14 +111,6 @@ public:
 	UPROPERTY(EditAnywhere, BlueprintReadWrite, Category = "Settings")
 	TArray<TScriptInterface<IMovieSceneConsoleVariableTrackInterface>> ConsoleVariablePresets;
 	
-	/** 
-	* An array of key/value pairs for console variable name and the value you wish to set for that cvar.
-	* The existing value will automatically be cached and restored afterwards.
-	*/
-	UE_DEPRECATED(5.2, "ConsoleVariables has been deprecated. Please use the console variable getters/setters (GetConsoleVariables(), RemoveConsoleVariable(), AddOrUpdateConsoleVariable(), and UpdateConsoleVariableEnableState()) instead. If scripting modifies ConsoleVariables, MRQ must be run in PIE to ensure changes are migrated.")
-	UPROPERTY(meta = (DeprecatedProperty, DeprecationMessage = "ConsoleVariables has been deprecated. Please use the console variable getters/setters (GetConsoleVariables(), RemoveConsoleVariable(), AddOrUpdateConsoleVariable(), and UpdateConsoleVariableEnableState()) instead. If scripting modifies ConsoleVariables, MRQ must be run in PIE to ensure changes are migrated."))
-	TMap<FString, float> ConsoleVariables_DEPRECATED;
-
 	/**
 	* An array of console commands to execute when this shot is started. If you need to restore the value 
 	* after the shot, add a matching entry in the EndConsoleCommands array. Because they are commands
@@ -183,9 +175,6 @@ public:
 	// End UObject interface
 
 private:
-	/** If the (deprecated) ConsoleVariables property contains any data, merge it into the new CVars property. */
-	void MergeInOldConsoleVariables();
-	
 	/** Merge together preset and override cvars into MergedConsoleVariables. Discards result of a prior merge (if any). */
 	void MergeInPresetConsoleVariables();
 
diff --git a/Engine/Plugins/MovieScene/SequencerScripting/Source/SequencerScripting/Public/ExtensionLibraries/MovieSceneFolderExtensions.h b/Engine/Plugins/MovieScene/SequencerScripting/Source/SequencerScripting/Public/ExtensionLibraries/MovieSceneFolderExtensions.h
--- a/Engine/Plugins/MovieScene/SequencerScripting/Source/SequencerScripting/Public/ExtensionLibraries/MovieSceneFolderExtensions.h
+++ b/Engine/Plugins/MovieScene/SequencerScripting/Source/SequencerScripting/Public/ExtensionLibraries/MovieSceneFolderExtensions.h
@@ -96,10 +96,6 @@ public:
 	UFUNCTION(BlueprintCallable, BlueprintPure, Category = "Sequencer|Sequence", meta = (ScriptMethod))
 	static TArray<UMovieSceneTrack*> GetChildTracks(UMovieSceneFolder* Folder);
 	
-	UE_DEPRECATED(5.2, "GetChildMasterTracks is deprecated. Please use GetChildTracks instead")
-	UFUNCTION(BlueprintCallable, BlueprintPure, Category = "Sequencer|Sequence", meta = (ScriptMethod, DeprecatedFunction, DeprecationMessage = "GetChildMasterTracks is deprecated. Please use GetChildTracks instead"))
-	static TArray<UMovieSceneTrack*> GetChildMasterTracks(UMovieSceneFolder* Folder) { return GetChildTracks(Folder); }
-
 	/**
 	 * Add a track to this folder
 	 *
@@ -110,10 +106,6 @@ public:
 	UFUNCTION(BlueprintCallable, Category = "Sequencer|Sequence", meta = (ScriptMethod))
 	static bool AddChildTrack(UMovieSceneFolder* Folder, UMovieSceneTrack* InTrack);
 	
-	UE_DEPRECATED(5.2, "AddChildMasterTrack is deprecated. Please use AddChildTrack instead")
-	UFUNCTION(BlueprintCallable, Category = "Sequencer|Sequence", meta = (ScriptMethod, DeprecatedFunction, DeprecationMessage = "AddChildMasterTrack is deprecated. Please use AddChildTrack instead"))
-	static bool AddChildMasterTrack(UMovieSceneFolder* Folder, UMovieSceneTrack* InTrack) { return AddChildTrack(Folder, InTrack); }
-
 	/**
 	 * Remove a track from the given folder
 	 *
@@ -124,10 +116,6 @@ public:
 	UFUNCTION(BlueprintCallable, Category = "Sequencer|Sequence", meta=(ScriptMethod))
 	static bool RemoveChildTrack(UMovieSceneFolder* Folder, UMovieSceneTrack* InTrack);
 
-	UE_DEPRECATED(5.2, "RemoveChildMasterTrack is deprecated. Please use RemoveChildTrack instead")
-	UFUNCTION(BlueprintCallable, Category = "Sequencer|Sequence", meta = (ScriptMethod, DeprecatedFunction, DeprecationMessage = "RemoveChildMasterTrack is deprecated. Please use RemoveChildTrack instead"))
-	static bool RemoveChildMasterTrack(UMovieSceneFolder* Folder, UMovieSceneTrack* InTrack) { return RemoveChildTrack(Folder, InTrack); }
-
 	/**
 	 * Get the object bindings contained by this folder
 	 *
diff --git a/Engine/Plugins/MovieScene/SequencerScripting/Source/SequencerScripting/Public/ExtensionLibraries/MovieSceneSequenceExtensions.h b/Engine/Plugins/MovieScene/SequencerScripting/Source/SequencerScripting/Public/ExtensionLibraries/MovieSceneSequenceExtensions.h
--- a/Engine/Plugins/MovieScene/SequencerScripting/Source/SequencerScripting/Public/ExtensionLibraries/MovieSceneSequenceExtensions.h
+++ b/Engine/Plugins/MovieScene/SequencerScripting/Source/SequencerScripting/Public/ExtensionLibraries/MovieSceneSequenceExtensions.h
@@ -46,10 +46,6 @@ public:
 	UFUNCTION(BlueprintCallable, Category="Sequencer|Sequence", meta=(ScriptMethod))
 	static TArray<UMovieSceneTrack*> GetTracks(UMovieSceneSequence* Sequence);
 
-	UE_DEPRECATED(5.2, "GetMasterTracks is deprecated. Please use GetTracks instead")
-	UFUNCTION(BlueprintCallable, Category = "Sequencer|Sequence", meta = (ScriptMethod, DeprecatedFunction, DeprecationMessage = "GetMasterTracks is deprecated. Please use GetTracks instead"))
-	static TArray<UMovieSceneTrack*> GetMasterTracks(UMovieSceneSequence* Sequence) { return GetTracks(Sequence); }
-
 	/**
 	 * Find all tracks of the specified type
 	 *
@@ -60,10 +56,6 @@ public:
 	UFUNCTION(BlueprintCallable, Category = "Sequencer|Sequence", meta = (ScriptMethod, DeterminesOutputType = "TrackType"))
 	static TArray<UMovieSceneTrack*> FindTracksByType(UMovieSceneSequence* Sequence, TSubclassOf<UMovieSceneTrack> TrackType);
 
-	UE_DEPRECATED(5.2, "FindMasterTracksByType is deprecated. Please use FindTracksByType instead")
-	UFUNCTION(BlueprintCallable, Category = "Sequencer|Sequence", meta = (ScriptMethod, DeterminesOutputType = "TrackType", DeprecatedFunction, DeprecationMessage = "FindMasterTracksByType is deprecated. Please use FindTracksByType instead"))
-	static TArray<UMovieSceneTrack*> FindMasterTracksByType(UMovieSceneSequence* Sequence, TSubclassOf<UMovieSceneTrack> TrackType) { return FindTracksByType(Sequence, TrackType); }
-
 	/**
 	 * Find all tracks of the specified type, not allowing sub-classed types
 	 *
@@ -74,10 +66,6 @@ public:
 	UFUNCTION(BlueprintCallable, Category="Sequencer|Sequence", meta=(ScriptMethod, DeterminesOutputType="TrackType"))
 	static TArray<UMovieSceneTrack*> FindTracksByExactType(UMovieSceneSequence* Sequence, TSubclassOf<UMovieSceneTrack> TrackType);
 
-	UE_DEPRECATED(5.2, "FindMasterTracksByExactType is deprecated. Please use FindTracksByExactType instead")
-	UFUNCTION(BlueprintCallable, Category = "Sequencer|Sequence", meta = (ScriptMethod, DeterminesOutputType = "TrackType", DeprecatedFunction, DeprecationMessage = "FindMasterTracksByExactType is deprecated. Please use FindTracksByExactType instead"))
-	static TArray<UMovieSceneTrack*> FindMasterTracksByExactType(UMovieSceneSequence* Sequence, TSubclassOf<UMovieSceneTrack> TrackType) { return FindTracksByExactType(Sequence, TrackType); }
-
 	/**
 	 * Add a new track of the specified type
 	 *
@@ -88,10 +76,6 @@ public:
 	UFUNCTION(BlueprintCallable, Category = "Sequencer|Sequence", meta = (ScriptMethod, DeterminesOutputType = "TrackType"))
 	static UMovieSceneTrack* AddTrack(UMovieSceneSequence* Sequence, TSubclassOf<UMovieSceneTrack> TrackType);
 
-	UE_DEPRECATED(5.2, "AddMasterTrack is deprecated. Please use AddTrack instead")
-	UFUNCTION(BlueprintCallable, Category = "Sequencer|Sequence", meta = (ScriptMethod, DeterminesOutputType = "TrackType", DeprecatedFunction, DeprecationMessage = "AddMasterTrack is deprecated. Please use AddTrack instead"))
-	static UMovieSceneTrack* AddMasterTrack(UMovieSceneSequence* Sequence, TSubclassOf<UMovieSceneTrack> TrackType) { return AddTrack(Sequence, TrackType); }
-
 	/**
 	 * Removes a track
 	 *
@@ -102,10 +86,6 @@ public:
 	UFUNCTION(BlueprintCallable, Category="Sequencer|Sequence", meta=(ScriptMethod))
 	static bool RemoveTrack(UMovieSceneSequence* Sequence, UMovieSceneTrack* Track);
 
-	UE_DEPRECATED(5.2, "RemoveMasterTrack is deprecated. Please use RemoveTrack instead")
-	UFUNCTION(BlueprintCallable, Category = "Sequencer|Sequence", meta = (ScriptMethod, DeprecatedFunction, DeprecationMessage = "RemoveMasterTrack is deprecated. Please use RemoveTrack instead"))
-	static bool RemoveMasterTrack(UMovieSceneSequence* Sequence, UMovieSceneTrack* Track) { return RemoveTrack(Sequence, Track); }
-
 	/**
 	 * Gets this sequence's display rate
 	 *
diff --git a/Engine/Source/Editor/MovieSceneTools/Private/MovieSceneToolHelpers.cpp b/Engine/Source/Editor/MovieSceneTools/Private/MovieSceneToolHelpers.cpp
--- a/Engine/Source/Editor/MovieSceneTools/Private/MovieSceneToolHelpers.cpp
+++ b/Engine/Source/Editor/MovieSceneTools/Private/MovieSceneToolHelpers.cpp
@@ -662,46 +662,6 @@ UMovieSceneSequence* MovieSceneToolHelpers::CreateSequence(FString& NewSequenceN
 	return Cast<UMovieSceneSequence>(NewAsset);
 }
 
-UMovieSceneSubSection* MovieSceneToolHelpers::CreateSubSequence(FString& NewSequenceName, FString& NewSequencePath, FFrameNumber NewSequenceStartTime, UMovieSceneSubTrack* SubTrack, UMovieSceneSubSection* SectionToDuplicate)
-{
-	// Create a new level sequence asset with the appropriate name
-	IAssetTools& AssetTools = FModuleManager::GetModuleChecked<FAssetToolsModule>("AssetTools").Get();
-
-	UObject* NewAsset = nullptr;
-	for (TObjectIterator<UClass> It; It; ++It)
-	{
-		UClass* CurrentClass = *It;
-		if (CurrentClass->IsChildOf(UFactory::StaticClass()) && !(CurrentClass->HasAnyClassFlags(CLASS_Abstract)))
-		{
-			UFactory* Factory = Cast<UFactory>(CurrentClass->GetDefaultObject());
-			if (Factory->CanCreateNew() && Factory->ImportPriority >= 0 && Factory->SupportedClass == ULevelSequence::StaticClass())
-			{
-				if (SectionToDuplicate != nullptr)
-				{
-					NewAsset = AssetTools.DuplicateAssetWithDialog(NewSequenceName, NewSequencePath, SectionToDuplicate->GetSequence());
-				}
-				else
-				{
-					NewAsset = AssetTools.CreateAssetWithDialog(NewSequenceName, NewSequencePath, ULevelSequence::StaticClass(), Factory);
-				}
-				break;
-			}
-		}
-	}
-
-	if (NewAsset == nullptr)
-	{
-		return nullptr;
-	}
-
-	UMovieSceneSequence* NewSequence = Cast<UMovieSceneSequence>(NewAsset);
-
-	int32 Duration = UE::MovieScene::DiscreteSize(SectionToDuplicate ? SectionToDuplicate->GetRange() : NewSequence->GetMovieScene()->GetPlaybackRange());
-
-	UMovieSceneSubSection* NewSection = SubTrack->AddSequence(NewSequence, NewSequenceStartTime, Duration);
-	return NewSection;
-}
-
 void MovieSceneToolHelpers::GatherTakes(const UMovieSceneSection* Section, TArray<FAssetData>& AssetData, uint32& OutCurrentTakeNumber)
 {
 	const UMovieSceneSubSection* SubSection = Cast<const UMovieSceneSubSection>(Section);
diff --git a/Engine/Source/Editor/MovieSceneTools/Public/MovieSceneToolHelpers.h b/Engine/Source/Editor/MovieSceneTools/Public/MovieSceneToolHelpers.h
--- a/Engine/Source/Editor/MovieSceneTools/Public/MovieSceneToolHelpers.h
+++ b/Engine/Source/Editor/MovieSceneTools/Public/MovieSceneToolHelpers.h
@@ -203,9 +203,6 @@ public:
 	 */
 	static UMovieSceneSequence* CreateSequence(FString& NewSequenceName, FString& NewSequencePath, UMovieSceneSubSection* SectionToDuplicate = nullptr);
 
-	UE_DEPRECATED(5.2, "CreateSubSequence has been deprecated in favor of CreateSequence so that asset creation can be invoked outside of a transaction (with the intent that it is undoable) and then adding the sequence to a subtrack can be invoked within a transaction")
-	static UMovieSceneSubSection* CreateSubSequence(FString& NewSequenceName, FString& NewSequencePath, FFrameNumber NewSequenceStartTime, UMovieSceneSubTrack* SubTrack, UMovieSceneSubSection* SectionToDuplicate = nullptr);
-
 	/**
 	 * Gather takes - level sequence assets that have the same shot prefix and shot number in the same asset path (directory)
 	 * 
diff --git a/Engine/Source/Editor/Sequencer/Public/MovieSceneTrackEditor.h b/Engine/Source/Editor/Sequencer/Public/MovieSceneTrackEditor.h
--- a/Engine/Source/Editor/Sequencer/Public/MovieSceneTrackEditor.h
+++ b/Engine/Source/Editor/Sequencer/Public/MovieSceneTrackEditor.h
@@ -164,44 +164,6 @@ public:
 		return Result;
 	}
 
-	template<typename TrackClass> struct
-	UE_DEPRECATED(5.2, "FFindOrCreateMasterTrackResult is deprecated. Please use FFindOrCreateRootTrackResult instead")
-	FFindOrCreateMasterTrackResult
-	{
-		FFindOrCreateMasterTrackResult() : Track(nullptr), bWasCreated(false) {}
-
-		TrackClass* Track;
-		bool bWasCreated;
-	};
-
-	/**
-	 * Find or add a track of the specified type in the focused movie scene.
-	 *
-	 * @param TrackClass The class of the track to find or add.
-	 * @return The track results.
-	 */
-	PRAGMA_DISABLE_DEPRECATION_WARNINGS // Suppress compiler warning on return of deprecated function
-	template<typename TrackClass>
-	UE_DEPRECATED(5.2, "FindOrCreateMasterTrack is deprecated. Please use FindOrCreateRootTrack instead")
-	FFindOrCreateMasterTrackResult<TrackClass> FindOrCreateMasterTrack()
-	{
-		FFindOrCreateMasterTrackResult<TrackClass> Result;
-		bool bTrackExisted;
-
-		UMovieScene* MovieScene = GetSequencer()->GetFocusedMovieSceneSequence()->GetMovieScene();
-		Result.Track = MovieScene->FindTrack<TrackClass>();
-		bTrackExisted = Result.Track != nullptr;
-
-		if (Result.Track == nullptr)
-		{
-			Result.Track = MovieScene->AddTrack<TrackClass>();
-		}
-
-		Result.bWasCreated = bTrackExisted == false && Result.Track != nullptr;
-		return Result;
-	}
-	PRAGMA_ENABLE_DEPRECATION_WARNINGS
-
 	/** @return The sequencer bound to this handler */
 	const TSharedPtr<ISequencer> GetSequencer() const;
 
diff --git a/Engine/Source/Editor/UnrealEd/Private/Fbx/FbxAnimationExport.cpp b/Engine/Source/Editor/UnrealEd/Private/Fbx/FbxAnimationExport.cpp
--- a/Engine/Source/Editor/UnrealEd/Private/Fbx/FbxAnimationExport.cpp
+++ b/Engine/Source/Editor/UnrealEd/Private/Fbx/FbxAnimationExport.cpp
@@ -63,21 +63,6 @@ namespace UnFbx
 		return true;
 	}
 
-void FFbxExporter::ExportAnimSequenceToFbx(const UAnimSequence* AnimSeq,
-									 const USkeletalMesh* SkelMesh,
-									 TArray<FbxNode*>& BoneNodes,
-									 FbxAnimLayer* InAnimLayer,
-									 float AnimStartOffset,
-									 float AnimEndOffset,
-									 float AnimPlayRate,
-									 float StartTime)
-{
-	ExportAnimSequenceToFbx(AnimSeq, SkelMesh, BoneNodes, InAnimLayer,
-		AnimSeq->GetDataModel()->GetFrameRate().AsFrameTime(AnimStartOffset),
-		AnimSeq->GetDataModel()->GetFrameRate().AsFrameTime(AnimSeq->GetPlayLength() - AnimEndOffset),
-		AnimPlayRate, StartTime);
-}
-
 void FFbxExporter::ExportAnimSequenceToFbx(const UAnimSequence* AnimSeq, const USkeletalMesh* SkelMesh, TArray<FbxNode*>& BoneNodes, FbxAnimLayer* InAnimLayer, FFrameTime StartFrameTime, FFrameTime EndFrameTime, float FrameRateScale, float StartTime)
 {
 	// stack allocator for extracting curve
diff --git a/Engine/Source/Editor/UnrealEd/Private/FbxExporter.h b/Engine/Source/Editor/UnrealEd/Private/FbxExporter.h
--- a/Engine/Source/Editor/UnrealEd/Private/FbxExporter.h
+++ b/Engine/Source/Editor/UnrealEd/Private/FbxExporter.h
@@ -416,10 +416,6 @@ private:
 	 */
 	UNREALED_API void ExportAnimSequenceToFbx(const UAnimSequence* AnimSeq, const USkeletalMesh* SkelMesh, TArray<FbxNode*>& BoneNodes, FbxAnimLayer* AnimLayer, FFrameTime StartFrameTime, FFrameTime EndFrameTime, float FrameRateScale, float StartTime);
 
-	UE_DEPRECATED(5.1, "ExportAnimSequenceToFbx is deprecated, use different signature")
-	UNREALED_API void ExportAnimSequenceToFbx(const UAnimSequence* AnimSeq, const USkeletalMesh* SkelMesh, TArray<FbxNode*>& BoneNodes, FbxAnimLayer* AnimLayer,
-		float AnimStartOffset, float AnimEndOffset, float AnimPlayRate, float StartTime);
-
 	/**
 	 * Add the custom Curve data to the FbxAnimCurves passed in parameter by matching their name to the skeletal mesh custom curves.
 	 */
diff --git a/Engine/Source/Runtime/Engine/Classes/Curves/RichCurve.h b/Engine/Source/Runtime/Engine/Classes/Curves/RichCurve.h
--- a/Engine/Source/Runtime/Engine/Classes/Curves/RichCurve.h
+++ b/Engine/Source/Runtime/Engine/Classes/Curves/RichCurve.h
@@ -323,13 +323,9 @@ public:
 	ENGINE_API virtual void BakeCurve(float SampleRate, float FirstKeyTime, float LastKeyTime) final override;
 
 	/** Remove redundant keys, comparing against Tolerance */
-	UE_DEPRECATED(5.1, "FRichCurve::RemoveRedundantKeys is deprecated, use signature with additional SampleRate or RemoveRedundantAutoTangentKeys instead")
-	void RemoveRedundantKeys(float Tolerance) { RemoveRedundantAutoTangentKeys(Tolerance); }
 	ENGINE_API void RemoveRedundantAutoTangentKeys(float Tolerance);
 	ENGINE_API virtual void RemoveRedundantKeys(float Tolerance, FFrameRate SampleRate) final override;
 
-	UE_DEPRECATED(5.1, "FRichCurve::RemoveRedundantKeys is deprecated, use signature with additional SampleRate or RemoveRedundantAutoTangentKeys instead")
-	void RemoveRedundantKeys(float Tolerance, float FirstKeyTime, float LastKeyTime) { RemoveRedundantAutoTangentKeys(Tolerance, FirstKeyTime, LastKeyTime); }
 	ENGINE_API void RemoveRedundantAutoTangentKeys(float Tolerance, float FirstKeyTime, float LastKeyTime);
 	ENGINE_API virtual void RemoveRedundantKeys(float Tolerance, float FirstKeyTime, float LastKeyTime, FFrameRate SampleRate) final override;
 
diff --git a/Engine/Source/Runtime/LevelSequence/Public/LevelSequenceDirector.h b/Engine/Source/Runtime/LevelSequence/Public/LevelSequenceDirector.h
--- a/Engine/Source/Runtime/LevelSequence/Public/LevelSequenceDirector.h
+++ b/Engine/Source/Runtime/LevelSequence/Public/LevelSequenceDirector.h
@@ -33,10 +33,6 @@ public:
 	UFUNCTION(BlueprintCallable, Category = "Sequencer|Director")
 	LEVELSEQUENCE_API FQualifiedFrameTime GetRootSequenceTime() const;
 
-	UE_DEPRECATED(5.2, "GetMasterSequenceTime is deprecated. Please use GetRootSequenceTime instead")
-	UFUNCTION(BlueprintCallable, Category = "Sequencer|Director", meta = (DeprecatedFunction, DeprecationMessage = "Use GetRootSequenceTime"))
-	FQualifiedFrameTime GetMasterSequenceTime() const { return GetRootSequenceTime(); }
-
 	/**
 	 * Get the current time for this director's sub-sequence (or the root sequence, if this is a root sequence director)
 	 * @return The current playback position of this director's sequence
diff --git a/Engine/Source/Runtime/LevelSequence/Public/LevelSequencePlayer.h b/Engine/Source/Runtime/LevelSequence/Public/LevelSequencePlayer.h
--- a/Engine/Source/Runtime/LevelSequence/Public/LevelSequencePlayer.h
+++ b/Engine/Source/Runtime/LevelSequence/Public/LevelSequencePlayer.h
@@ -63,14 +63,6 @@ struct FLevelSequencePlayerSnapshot
 
 	UPROPERTY()
 	FMovieSceneSequenceID ShotID;
-
-private:
-
-	UPROPERTY()
-	FString MasterName_DEPRECATED;
-
-	UPROPERTY()
-	FQualifiedFrameTime MasterTime_DEPRECATED;
 };
 
 /**
diff --git a/Engine/Source/Runtime/MovieScene/Private/EntitySystem/MovieSceneEntityFactory.cpp b/Engine/Source/Runtime/MovieScene/Private/EntitySystem/MovieSceneEntityFactory.cpp
--- a/Engine/Source/Runtime/MovieScene/Private/EntitySystem/MovieSceneEntityFactory.cpp
+++ b/Engine/Source/Runtime/MovieScene/Private/EntitySystem/MovieSceneEntityFactory.cpp
@@ -265,16 +265,6 @@ void FEntityFactories::DefineComplexInclusiveComponents(const FComplexInclusivit
 	MutualInclusivityGraph.DefineComplexInclusionRule(InFilter, InComponents, MoveTemp(Params));
 }
 
-PRAGMA_DISABLE_DEPRECATION_WARNINGS
-void FEntityFactories::DefineComplexInclusiveComponents(const FComplexInclusivity& InInclusivity)
-{
-	for (FComponentMaskIterator It(InInclusivity.ComponentsToInclude.Iterate()); It; ++It)
-	{
-		MutualInclusivityGraph.DefineComplexInclusionRule(InInclusivity.Filter, { FComponentTypeID::FromBitIndex(It.GetIndex()) });
-	}
-}
-PRAGMA_ENABLE_DEPRECATION_WARNINGS
-
 int32 FEntityFactories::ComputeChildComponents(const FComponentMask& ParentComponentMask, FComponentMask& ChildComponentMask)
 {
 	int32 NumNewComponents = 0;
diff --git a/Engine/Source/Runtime/MovieScene/Private/Evaluation/MovieSceneEvaluationTemplateInstance.cpp b/Engine/Source/Runtime/MovieScene/Private/Evaluation/MovieSceneEvaluationTemplateInstance.cpp
--- a/Engine/Source/Runtime/MovieScene/Private/Evaluation/MovieSceneEvaluationTemplateInstance.cpp
+++ b/Engine/Source/Runtime/MovieScene/Private/Evaluation/MovieSceneEvaluationTemplateInstance.cpp
@@ -173,11 +173,6 @@ void FMovieSceneRootEvaluationTemplateInstance::Initialize(UMovieSceneSequence&
 	}
 }
 
-void FMovieSceneRootEvaluationTemplateInstance::Evaluate(FMovieSceneContext Context, IMovieScenePlayer& Player)
-{
-	EvaluateSynchronousBlocking(Context);
-}
-
 void FMovieSceneRootEvaluationTemplateInstance::EvaluateSynchronousBlocking(FMovieSceneContext Context, IMovieScenePlayer& Player)
 {
 	EvaluateSynchronousBlocking(Context);
diff --git a/Engine/Source/Runtime/MovieScene/Public/EntitySystem/MovieSceneEntityFactory.h b/Engine/Source/Runtime/MovieScene/Public/EntitySystem/MovieSceneEntityFactory.h
--- a/Engine/Source/Runtime/MovieScene/Public/EntitySystem/MovieSceneEntityFactory.h
+++ b/Engine/Source/Runtime/MovieScene/Public/EntitySystem/MovieSceneEntityFactory.h
@@ -30,12 +30,6 @@ struct FEntityRange;
 struct FEntityAllocation;
 struct FMutualComponentInitializers;
 
-struct UE_DEPRECATED(5.2, "Please use DefineComplexInclusiveComponents()") FComplexInclusivity
-{
-	FComplexInclusivityFilter Filter;
-	FComponentMask ComponentsToInclude;
-};
-
 /**
  * A class that contains all the component factory relationships.
  *
@@ -131,27 +125,6 @@ struct FEntityFactories
 
 	MOVIESCENE_API void RunInitializers(const FComponentMask& ParentType, const FComponentMask& ChildType, const FEntityAllocation* ParentAllocation, TArrayView<const int32> ParentAllocationOffsets, const FEntityRange& InChildEntityRange);
 
-public:
-
-	PRAGMA_DISABLE_DEPRECATION_WARNINGS
-		using FDeprecatedComplexInclusivity = FComplexInclusivity;
-	PRAGMA_ENABLE_DEPRECATION_WARNINGS
-	/**
-	 * Specifies that if an entity matches the given filter, the specified components should be created on it.
-	 */
-	template<typename... ComponentTypes>
-	UE_DEPRECATED(5.2, "Please use DefineComplexInclusiveComponents(const FComplexInclusivityFilter&, initializer_list<FComponentTypeID>)")
-	void DefineComplexInclusiveComponents(const FComplexInclusivityFilter& InFilter, FComponentTypeID InComponent, ComponentTypes... InComponents)
-	{
-		DefineComplexInclusiveComponents(InFilter, std::initializer_list<FComponentTypeID>({ InComponent, InComponents... }), FMutuallyInclusiveComponentParams());
-	}
-
-	/**
-	 * Defines a new complex inclusivity relationship. The helper methods above are easier and preferrable.
-	 */
-	UE_DEPRECATED(5.2, "Please use DefineComplexInclusiveComponents(const FComplexInclusivityFilter&, FComponentTypeID)")
-	MOVIESCENE_API void DefineComplexInclusiveComponents(const FDeprecatedComplexInclusivity& InInclusivity);
-
 private:
 
 	TArray<TInlineValue<FChildEntityInitializer>> ChildInitializers;
diff --git a/Engine/Source/Runtime/MovieScene/Public/EntitySystem/MovieSceneEntitySystemTask.h b/Engine/Source/Runtime/MovieScene/Public/EntitySystem/MovieSceneEntitySystemTask.h
--- a/Engine/Source/Runtime/MovieScene/Public/EntitySystem/MovieSceneEntitySystemTask.h
+++ b/Engine/Source/Runtime/MovieScene/Public/EntitySystem/MovieSceneEntitySystemTask.h
@@ -263,12 +263,6 @@ struct TEntityTaskComponents : TEntityTaskComponentsImpl<TMakeIntegerSequence<in
 		return TFilteredEntityTask< T... >(*this, InFilter);
 	}
 
-	UE_DEPRECATED(5.2, "This function is not required.")
-	TEntityTaskComponents< T... >& SetCurrentThread(ENamedThreads::Type InCurrentThread)
-	{
-		return *this;
-	}
-
 	/**
 	 * Assign a desired thread for this task to run on
 	 */
@@ -1111,12 +1105,6 @@ struct TFilteredEntityTask
 		return *this;
 	}
 
-	UE_DEPRECATED(5.2, "This function is not required.")
-	TFilteredEntityTask< T... >& SetCurrentThread(ENamedThreads::Type InCurrentThread)
-	{
-		return *this;
-	}
-
 	/**
 	 * Assign a desired thread for this task to run on
 	 */
diff --git a/Engine/Source/Runtime/MovieScene/Public/Evaluation/MovieSceneEvaluationTemplateInstance.h b/Engine/Source/Runtime/MovieScene/Public/Evaluation/MovieSceneEvaluationTemplateInstance.h
--- a/Engine/Source/Runtime/MovieScene/Public/Evaluation/MovieSceneEvaluationTemplateInstance.h
+++ b/Engine/Source/Runtime/MovieScene/Public/Evaluation/MovieSceneEvaluationTemplateInstance.h
@@ -62,9 +62,6 @@ public:
 	UE_DEPRECATED(5.4, "Use the version without the player parameter")
 	MOVIESCENE_API void EvaluateSynchronousBlocking(FMovieSceneContext Context, IMovieScenePlayer& Player);
 
-	UE_DEPRECATED(5.1, "Use EvaluateSynchronousBlocking instead.")
-	MOVIESCENE_API void Evaluate(FMovieSceneContext Context, IMovieScenePlayer& Player);
-
 	UE_DEPRECATED(5.4, "Director instances are now auto-managed via FSequenceDirectorPlaybackCapability")
 	MOVIESCENE_API void ResetDirectorInstances();
 
diff --git a/Engine/Source/Runtime/MovieScene/Public/MovieScene.h b/Engine/Source/Runtime/MovieScene/Public/MovieScene.h
--- a/Engine/Source/Runtime/MovieScene/Public/MovieScene.h
+++ b/Engine/Source/Runtime/MovieScene/Public/MovieScene.h
@@ -609,9 +609,6 @@ public:
 	 */
 	MOVIESCENE_API UMovieSceneTrack* AddTrack(TSubclassOf<UMovieSceneTrack> TrackClass);
 	
-	UE_DEPRECATED(5.2, "AddMasterTrack is deprecated. Please use AddTrack instead")
-	UMovieSceneTrack* AddMasterTrack(TSubclassOf<UMovieSceneTrack> TrackClass) { return AddTrack(TrackClass); }
-
 	/**
 	 * Adds a track.
 	 *
@@ -627,10 +624,6 @@ public:
 		return Cast<TrackClass>(AddTrack(TrackClass::StaticClass()));
 	}
 
-	template<typename TrackClass>
-	UE_DEPRECATED(5.2, "AddMasterTrack is deprecated. Please use AddTrack instead")
-	TrackClass* AddMasterTrack() { return AddTrack<TrackClass>(); }
-
 	/**
 	* Adds a given track as a track
 	*
@@ -640,9 +633,6 @@ public:
 	*/
 	MOVIESCENE_API bool AddGivenTrack(UMovieSceneTrack* InTrack);
 
-	UE_DEPRECATED(5.2, "AddGivenMasterTrack is deprecated. Please use AddGivenTrack instead")
-	bool AddGivenMasterTrack(UMovieSceneTrack* InTrack) { return AddGivenTrack(InTrack); }
-
 	/**
 	 * Finds a track (one not bound to a runtime objects).
 	 *
@@ -652,9 +642,6 @@ public:
 	 */
 	MOVIESCENE_API UMovieSceneTrack* FindTrack(TSubclassOf<UMovieSceneTrack> TrackClass) const;
 
-	UE_DEPRECATED(5.2, "FindMasterTrack is deprecated. Please use FindTrack instead")
-	UMovieSceneTrack* FindMasterTrack(TSubclassOf<UMovieSceneTrack> TrackClass) const { return FindTrack(TrackClass); }
-
 	/**
 	 * Finds a track (one not bound to a runtime objects).
 	 *
@@ -668,10 +655,6 @@ public:
 		return Cast<TrackClass>(FindTrack(TrackClass::StaticClass()));
 	}
 
-	template<typename TrackClass>
-	UE_DEPRECATED(5.2, "FindMasterTrack is deprecated. Please use FindTrack instead")
-	TrackClass* FindMasterTrack() const { return FindTrack<TrackClass>(); }
-
 	/**
 	 * Get all tracks.
 	 *
@@ -683,9 +666,6 @@ public:
 		return Tracks;
 	}
 
-	UE_DEPRECATED(5.2, "GetMasterTracks is deprecated. Please use GetTracks instead")
-	const TArray<UMovieSceneTrack*>& GetMasterTracks() const { return GetTracks(); }
-
 	/**
 	 * Check whether the specified track is a track in this movie scene.
 	 *
@@ -694,12 +674,6 @@ public:
 	 */
 	MOVIESCENE_API bool ContainsTrack(const UMovieSceneTrack& Track) const;
 
-	UE_DEPRECATED(5.2, "IsAMasterTrack is deprecated. Please use ContainsTrack instead")
-	bool IsAMasterTrack(const UMovieSceneTrack& Track) const { return ContainsTrack(Track); }
-
-	UE_DEPRECATED(5.2, "RemoveMasterTrack is deprecated. Please use RemoveTrack instead")
-	bool RemoveMasterTrack(UMovieSceneTrack& Track) { return RemoveTrack(Track); }
-
 	/**
 	 * Move all the contents (tracks, child bindings) of the specified binding ID onto another
 	 *
diff --git a/Engine/Source/Runtime/MovieScene/Public/MovieSceneFolder.h b/Engine/Source/Runtime/MovieScene/Public/MovieSceneFolder.h
--- a/Engine/Source/Runtime/MovieScene/Public/MovieSceneFolder.h
+++ b/Engine/Source/Runtime/MovieScene/Public/MovieSceneFolder.h
@@ -50,27 +50,15 @@ class UMovieSceneFolder : public UObject
 	/** Gets the tracks contained by this folder. */
 	MOVIESCENE_API const TArray<UMovieSceneTrack*>& GetChildTracks() const;
 
-	UE_DEPRECATED(5.2, "GetChildMasterTracks is deprecated. Please use GetChildTracks instead")
-	const TArray<UMovieSceneTrack*>& GetChildMasterTracks() const { return GetChildTracks(); }
-
 	/** Adds a track to this folder. Automatically calls Modify() on the folder object. */
 	MOVIESCENE_API void AddChildTrack( UMovieSceneTrack* InTrack );
 
-	UE_DEPRECATED(5.2, "AddChildMasterTrack is deprecated. Please use AddChildTrack instead")
-	void AddChildMasterTrack(UMovieSceneTrack* InMasterTrack) { AddChildTrack(InMasterTrack); }
-
 	/** Removes a track from this folder. Automatically calls Modify() on the folder object. */
 	MOVIESCENE_API void RemoveChildTrack( UMovieSceneTrack* InTrack );
 
-	UE_DEPRECATED(5.2, "RemoveChildMasterTrack is deprecated. Please use RemoveChildTrack instead")
-	void RemoveChildMasterTrack(UMovieSceneTrack* InMasterTrack) { return RemoveChildTrack(InMasterTrack); }
-
 	/** Clear all child tracks from this folder. */
 	MOVIESCENE_API void ClearChildTracks();
 
-	UE_DEPRECATED(5.2, "ClearChildMasterTracks is deprecated. Please use ClearChildTracks instead")
-	void ClearChildMasterTracks() { ClearChildTracks(); }
-
 	/** Gets the guids for the object bindings contained by this folder. */
 	MOVIESCENE_API const TArray<FGuid>& GetChildObjectBindings() const;
 
diff --git a/Engine/Source/Runtime/MovieSceneTracks/Public/Sections/MovieSceneDataLayerSection.h b/Engine/Source/Runtime/MovieSceneTracks/Public/Sections/MovieSceneDataLayerSection.h
--- a/Engine/Source/Runtime/MovieSceneTracks/Public/Sections/MovieSceneDataLayerSection.h
+++ b/Engine/Source/Runtime/MovieSceneTracks/Public/Sections/MovieSceneDataLayerSection.h
@@ -62,17 +62,6 @@ private:
 	virtual void ImportEntityImpl(UMovieSceneEntitySystemLinker* EntityLinker, const FEntityImportParams& Params, FImportedEntity* OutImportedEntity) override;
 
 private:
-	UE_DEPRECATED(5.1, "Use GetDataLayerAssets instead")
-	UFUNCTION(BlueprintPure, Category = "Sequencer|Section")
-	const TArray<FActorDataLayer>& GetDataLayers() const { return DataLayers; }
-
-	UE_DEPRECATED(5.1, "Use SetDataLayerAssets instead")
-	UFUNCTION(BlueprintCallable, Category = "Sequencer|Section")
-	void SetDataLayers(const TArray<FActorDataLayer>& InDataLayers) { DataLayers = InDataLayers; }
-
-	UPROPERTY(meta = (DeprecatedProperty, DeprecationMessage = "Use DataLayer Assets Instead"))
-	TArray<FActorDataLayer> DataLayers;
-
 	/** A list of data layers that should be loaded or unloaded by this section */
 	UPROPERTY(EditAnywhere, Category = DataLayer)
 	TArray<TObjectPtr<UDataLayerAsset>> DataLayerAssets;
diff --git a/Engine/Source/Runtime/MovieSceneTracks/Public/Tracks/MovieSceneAudioTrack.h b/Engine/Source/Runtime/MovieSceneTracks/Public/Tracks/MovieSceneAudioTrack.h
--- a/Engine/Source/Runtime/MovieSceneTracks/Public/Tracks/MovieSceneAudioTrack.h
+++ b/Engine/Source/Runtime/MovieSceneTracks/Public/Tracks/MovieSceneAudioTrack.h
@@ -39,9 +39,6 @@ public:
 		return AudioSections;
 	}
 
-	UE_DEPRECATED(5.2, "IsAMasterTrack is deprecated. Please use Cast<UMovieScene>(GetOuter())->ContainsTrack instead")
-	MOVIESCENETRACKS_API bool IsAMasterTrack() const;
-
 public:
 
 	// UMovieSceneTrack interface
