Sequencer: Remove use of transient player which was incorrectly initialized and not returning the bound objects. Instead, use a transient playback state, similar to 30020582

diff --git a/Engine/Plugins/MovieScene/SequencerScripting/Source/SequencerScripting/Private/ExtensionLibraries/MovieSceneSequenceExtensions.cpp b/Engine/Plugins/MovieScene/SequencerScripting/Source/SequencerScripting/Private/ExtensionLibraries/MovieSceneSequenceExtensions.cpp
--- a/Engine/Plugins/MovieScene/SequencerScripting/Source/SequencerScripting/Private/ExtensionLibraries/MovieSceneSequenceExtensions.cpp
+++ b/Engine/Plugins/MovieScene/SequencerScripting/Source/SequencerScripting/Private/ExtensionLibraries/MovieSceneSequenceExtensions.cpp
@@ -821,21 +821,18 @@ TArray<UObject*> UMovieSceneSequenceExtensions::LocateBoundObjects(UMovieSceneSe
 		return TArray<UObject*>();
 	}
 
-	class FTransientPlayer : public IMovieScenePlayer
-	{
-	public:
-		FMovieSceneRootEvaluationTemplateInstance Template;
-		virtual FMovieSceneRootEvaluationTemplateInstance& GetEvaluationTemplate() override { return Template; }
-		virtual void UpdateCameraCut(UObject* CameraObject, const EMovieSceneCameraCutParams& CameraCutParams) override {}
-		virtual void SetViewportSettings(const TMap<FViewportClient*, EMovieSceneViewportParams>& ViewportParamsMap) override {}
-		virtual void GetViewportSettings(TMap<FViewportClient*, EMovieSceneViewportParams>& ViewportParamsMap) const override {}
-		virtual EMovieScenePlayerStatus::Type GetPlaybackStatus() const { return EMovieScenePlayerStatus::Stopped; }
-		virtual void SetPlaybackStatus(EMovieScenePlayerStatus::Type InPlaybackStatus) override {}
-	} Player;
+	using namespace UE::MovieScene;
+
+	FSharedPlaybackStateCreateParams CreateParams;
+	CreateParams.PlaybackContext = Context;
+	TSharedRef<FSharedPlaybackState> TransientPlaybackState = MakeShared<FSharedPlaybackState>(*Sequence, CreateParams);
+
+	FMovieSceneEvaluationState State;
+	TransientPlaybackState->AddCapabilityRaw(&State);
+	State.AssignSequence(MovieSceneSequenceID::Root, *Sequence, TransientPlaybackState);
 
-	Player.State.AssignSequence(MovieSceneSequenceID::Root, *Sequence, Player);
+	TArrayView<TWeakObjectPtr<>> Objects = State.FindBoundObjects(InBinding.BindingID, MovieSceneSequenceID::Root, TransientPlaybackState);
 
-	TArrayView<TWeakObjectPtr<>> Objects = Player.FindBoundObjects(InBinding.BindingID, MovieSceneSequenceID::Root);
 	TArray<UObject*> Result;
 	for (TWeakObjectPtr<> WeakObject : Objects)
 	{
diff --git a/Engine/Plugins/MovieScene/SequencerScripting/Source/SequencerScripting/SequencerScripting.Build.cs b/Engine/Plugins/MovieScene/SequencerScripting/Source/SequencerScripting/SequencerScripting.Build.cs
--- a/Engine/Plugins/MovieScene/SequencerScripting/Source/SequencerScripting/SequencerScripting.Build.cs
+++ b/Engine/Plugins/MovieScene/SequencerScripting/Source/SequencerScripting/SequencerScripting.Build.cs
@@ -38,6 +38,7 @@ public class SequencerScripting : ModuleRules
 				"Slate",
 				"SlateCore",
                 "LevelSequence",
+				"UniversalObjectLocator",
 			}
 		);
 
