Sequencer: Remove 5.1, 5.3 deprecations

#rnx

diff --git a/Engine/Plugins/MovieScene/LevelSequenceEditor/Source/LevelSequenceEditor/Private/LevelSequenceEditorSubsystem.cpp b/Engine/Plugins/MovieScene/LevelSequenceEditor/Source/LevelSequenceEditor/Private/LevelSequenceEditorSubsystem.cpp
--- a/Engine/Plugins/MovieScene/LevelSequenceEditor/Source/LevelSequenceEditor/Private/LevelSequenceEditorSubsystem.cpp
+++ b/Engine/Plugins/MovieScene/LevelSequenceEditor/Source/LevelSequenceEditor/Private/LevelSequenceEditorSubsystem.cpp
@@ -1635,57 +1635,6 @@ void ULevelSequenceEditorSubsystem::BakeTransformInternal()
 	BakeWidget->OpenDialog(true);
 }
 
-void ULevelSequenceEditorSubsystem::BakeTransform(const TArray<FMovieSceneBindingProxy>& ObjectBindings, const FFrameTime& BakeInTime, const FFrameTime& BakeOutTime, const FFrameTime& BakeInterval, const FMovieSceneScriptingParams& Params)
-{
-	TSharedPtr<ISequencer> Sequencer = GetActiveSequencer();
-	if (Sequencer == nullptr)
-	{
-		return;
-	}
-
-	UMovieScene* FocusedMovieScene = Sequencer->GetFocusedMovieSceneSequence()->GetMovieScene();
-	if (!FocusedMovieScene)
-	{
-		return;
-	}
-
-	if (FocusedMovieScene->IsReadOnly())
-	{
-		FSequencerUtilities::ShowReadOnlyError();
-		return;
-	}
-
-	if (ObjectBindings.Num() == 0)
-	{
-		return;
-	}
-
-	FFrameRate TickResolution = FocusedMovieScene->GetTickResolution();
-	FFrameRate DisplayRate = FocusedMovieScene->GetDisplayRate();
-	FFrameTime InFrame = BakeInTime;
-	FFrameTime OutFrame = BakeOutTime;
-	FFrameTime Interval = BakeInterval;
-
-	if (Params.TimeUnit == EMovieSceneTimeUnit::DisplayRate)
-	{
-		InFrame = ConvertFrameTime(BakeInTime, DisplayRate, TickResolution);
-		OutFrame = ConvertFrameTime(BakeOutTime, DisplayRate, TickResolution);
-		Interval = ConvertFrameTime(BakeInterval, DisplayRate, TickResolution);
-	}
-
-	FBakingAnimationKeySettings Settings;
-	FFrameTime OneFrame = ConvertFrameTime(1, DisplayRate, TickResolution);
-	//we never want subframes when baking so use frame increments
-	Settings.FrameIncrement = Interval.FrameNumber.Value / OneFrame.FrameNumber.Value;
-	Settings.FrameIncrement = Settings.FrameIncrement <= 0 ? 1 : Settings.FrameIncrement;
-	Settings.StartFrame = InFrame.GetFrame();
-	Settings.EndFrame = OutFrame.GetFrame();
-	Settings.BakingKeySettings = EBakingKeySettings::AllFrames;
-	FMovieSceneScriptingParams NewParams;
-	NewParams.TimeUnit = EMovieSceneTimeUnit::TickResolution;
-	BakeTransformWithSettings(ObjectBindings, Settings, NewParams);
-}
-
 void ULevelSequenceEditorSubsystem::CalculateFramesPerGuid(TSharedPtr<ISequencer>& Sequencer, const FBakingAnimationKeySettings& InSettings, TMap<FGuid, ULevelSequenceEditorSubsystem::FBakeData>& OutBakeDataMap,
 	TSortedMap<FFrameNumber, FFrameNumber>& OutFrameMap)
 {
diff --git a/Engine/Plugins/MovieScene/LevelSequenceEditor/Source/LevelSequenceEditor/Public/LevelSequenceEditorSubsystem.h b/Engine/Plugins/MovieScene/LevelSequenceEditor/Source/LevelSequenceEditor/Public/LevelSequenceEditorSubsystem.h
--- a/Engine/Plugins/MovieScene/LevelSequenceEditor/Source/LevelSequenceEditor/Public/LevelSequenceEditorSubsystem.h
+++ b/Engine/Plugins/MovieScene/LevelSequenceEditor/Source/LevelSequenceEditor/Public/LevelSequenceEditorSubsystem.h
@@ -228,10 +228,6 @@ public:
 	void SyncSectionsUsingSourceTimecode(const TArray<UMovieSceneSection*>& Sections);
 
 	/** Bake transform */
-	UE_DEPRECATED(5.3, "Use ULevelSequenceEditorSubsystem::BakeTransformWithSettings instead")
-	UFUNCTION(BlueprintCallable, Category = "Level Sequence Editor")
-	void BakeTransform(const TArray<FMovieSceneBindingProxy>& ObjectBindings, const FFrameTime& BakeInTime, const FFrameTime& BakeOutTime, const FFrameTime& BakeInterval, const FMovieSceneScriptingParams& Params = FMovieSceneScriptingParams());
-
 	UFUNCTION(BlueprintCallable, Category = "Level Sequence Editor")
 	bool BakeTransformWithSettings(const TArray<FMovieSceneBindingProxy>& ObjectBindings, const FBakingAnimationKeySettings& InSettings, const FMovieSceneScriptingParams& Params = FMovieSceneScriptingParams());
 
diff --git a/Engine/Plugins/MovieScene/SequencerScripting/Source/SequencerScriptingEditor/Private/SequencerTools.cpp b/Engine/Plugins/MovieScene/SequencerScripting/Source/SequencerScriptingEditor/Private/SequencerTools.cpp
--- a/Engine/Plugins/MovieScene/SequencerScripting/Source/SequencerScriptingEditor/Private/SequencerTools.cpp
+++ b/Engine/Plugins/MovieScene/SequencerScripting/Source/SequencerScriptingEditor/Private/SequencerTools.cpp
@@ -40,147 +40,6 @@
 
 #define LOCTEXT_NAMESPACE "SequencerTools"
 
-bool USequencerToolsFunctionLibrary::RenderMovie(UMovieSceneCapture* InCaptureSettings, FOnRenderMovieStopped OnFinishedCallback)
-{
-	IMovieSceneCaptureDialogModule& MovieSceneCaptureModule = FModuleManager::Get().LoadModuleChecked<IMovieSceneCaptureDialogModule>("MovieSceneCaptureDialog");
-	
-	// Because this comes from the Python/BP layer we need to soft-validate the state before we pass it onto functions that do a assert-based validation.
-	if (!InCaptureSettings)
-	{
-		FFrame::KismetExecutionMessage(TEXT("Cannot start Render Sequence to Movie with null capture settings."), ELogVerbosity::Error);
-		return false;
-	}
-
-PRAGMA_DISABLE_DEPRECATION_WARNINGS
-	if (IsRenderingMovie())
-	{
-		FFrame::KismetExecutionMessage(TEXT("Capture already in progress."), ELogVerbosity::Error);
-		return false;
-	}
-PRAGMA_ENABLE_DEPRECATION_WARNINGS
-
-	// If they're capturing a level sequence we'll do some additional checking as there are more parameters on the Automated Level Sequence capture.
-	UAutomatedLevelSequenceCapture* LevelSequenceCapture = Cast<UAutomatedLevelSequenceCapture>(InCaptureSettings);
-	if (LevelSequenceCapture)
-	{
-		if (!LevelSequenceCapture->LevelSequenceAsset.IsValid())
-		{
-			// UE_LOG(LogTemp, Warning, TEXT("No Level Sequence Asset specified in UAutomatedLevelSequenceCapture."));
-			FFrame::KismetExecutionMessage(TEXT("No Level Sequence Asset specified in UAutomatedLevelSequenceCapture."), ELogVerbosity::Error);
-			return false;
-		}
-
-		if (!LevelSequenceCapture->bUseCustomStartFrame && !LevelSequenceCapture->bUseCustomEndFrame)
-		{
-			// If they don't want to use a custom start/end frame we override the default values to be the length of the sequence, as the default is [0,1)
-			ULevelSequence* LevelSequence = Cast<ULevelSequence>(LevelSequenceCapture->LevelSequenceAsset.TryLoad());
-			if (!LevelSequence)
-			{
-				const FString ErrorMessage = FString::Printf(TEXT("Specified Level Sequence Asset failed to load. Specified Asset Path: %s"), *LevelSequenceCapture->LevelSequenceAsset.GetAssetPathString());
-				FFrame::KismetExecutionMessage(*ErrorMessage, ELogVerbosity::Error);
-				return false;
-			}
-
-			FFrameRate DisplayRate = LevelSequence->GetMovieScene()->GetDisplayRate();
-			FFrameRate TickResolution = LevelSequence->GetMovieScene()->GetTickResolution();
-			
-			LevelSequenceCapture->Settings.CustomFrameRate = DisplayRate;
-			LevelSequenceCapture->Settings.bUseCustomFrameRate = true;
-			LevelSequenceCapture->Settings.bUseRelativeFrameNumbers = false;
-			TRange<FFrameNumber> Range = LevelSequence->GetMovieScene()->GetPlaybackRange();
-
-			FFrameNumber StartFrame = UE::MovieScene::DiscreteInclusiveLower(Range);
-			FFrameNumber EndFrame = UE::MovieScene::DiscreteExclusiveUpper(Range);
-
-			FFrameNumber RoundedStartFrame = FFrameRate::TransformTime(StartFrame, TickResolution, DisplayRate).CeilToFrame();
-			FFrameNumber RoundedEndFrame = FFrameRate::TransformTime(EndFrame, TickResolution, DisplayRate).CeilToFrame();
-
-			LevelSequenceCapture->CustomStartFrame = RoundedStartFrame;
-			LevelSequenceCapture->CustomEndFrame = RoundedEndFrame;
-		}
-	}
-
-	auto LocalCaptureStoppedCallback = [OnFinishedCallback](bool bSuccess)
-	{
-		OnFinishedCallback.ExecuteIfBound(bSuccess);
-	};
-
-	MovieSceneCaptureModule.StartCapture(InCaptureSettings);
-	MovieSceneCaptureModule.GetCurrentCapture()->CaptureStoppedDelegate.AddLambda(LocalCaptureStoppedCallback);
-	return true;
-}
-
-void USequencerToolsFunctionLibrary::CancelMovieRender()
-{
-	IMovieSceneCaptureDialogModule& MovieSceneCaptureModule = FModuleManager::Get().LoadModuleChecked<IMovieSceneCaptureDialogModule>("MovieSceneCaptureDialog");
-	TSharedPtr<FMovieSceneCaptureBase> CurrentCapture = MovieSceneCaptureModule.GetCurrentCapture();
-	if (CurrentCapture.IsValid())
-	{
-		// We just invoke the capture's Cancel function. This will cause a shut-down of the capture (the same as the UI)
-		// which will invoke all of the necessary callbacks as well. We don't null out CurrentCapture because that is done
-		// as the result of its shutdown callbacks.
-		CurrentCapture->Cancel();
-	}
-}
-
-TArray<FSequencerBoundObjects> USequencerToolsFunctionLibrary::GetBoundObjects(UWorld* InWorld, ULevelSequence* InSequence, const TArray<FMovieSceneBindingProxy>& InBindings, const FSequencerScriptingRange& InRange)
-{
-	ALevelSequenceActor* OutActor;
-	FMovieSceneSequencePlaybackSettings Settings;
-	FLevelSequenceCameraSettings CameraSettings;
-
-	ULevelSequencePlayer* Player = ULevelSequencePlayer::CreateLevelSequencePlayer(InWorld, InSequence, Settings, OutActor);
-
-	// Evaluation needs to occur in order to obtain spawnables
-	FFrameRate Resolution = InSequence->GetMovieScene()->GetTickResolution();
-	TRange<FFrameNumber> SpecifiedRange = InRange.ToNative(Resolution);
-	Player->SetPlaybackPosition(FMovieSceneSequencePlaybackParams(SpecifiedRange.GetLowerBoundValue().Value, EUpdatePositionMethod::Play));
-
-	FMovieSceneSequenceID SequenceId = Player->State.FindSequenceId(InSequence);
-
-	TArray<FSequencerBoundObjects> BoundObjects;
-	for (FMovieSceneBindingProxy Binding : InBindings)
-	{
-		FMovieSceneObjectBindingID ObjectBinding = UE::MovieScene::FFixedObjectBindingID(Binding.BindingID, SequenceId);
-		BoundObjects.Add(FSequencerBoundObjects(Binding, Player->GetBoundObjects(ObjectBinding)));
-	}
-
-	Player->Stop();
-	InWorld->DestroyActor(OutActor);
-	return BoundObjects;
-}
-
-
-TArray<FSequencerBoundObjects> USequencerToolsFunctionLibrary::GetObjectBindings(UWorld* InWorld, ULevelSequence* InSequence, const TArray<UObject*>& InObjects, const FSequencerScriptingRange& InRange)
-{
-	ALevelSequenceActor* OutActor;
-	FMovieSceneSequencePlaybackSettings Settings;
-	FLevelSequenceCameraSettings CameraSettings;
-
-	ULevelSequencePlayer* Player = ULevelSequencePlayer::CreateLevelSequencePlayer(InWorld, InSequence, Settings, OutActor);
-
-	FFrameRate Resolution = InSequence->GetMovieScene()->GetTickResolution();
-	TRange<FFrameNumber> SpecifiedRange = InRange.ToNative(Resolution);
-	Player->SetPlaybackPosition(FMovieSceneSequencePlaybackParams(SpecifiedRange.GetLowerBoundValue().Value, EUpdatePositionMethod::Play));
-
-	TArray<FSequencerBoundObjects> BoundObjects;
-
-	for (UObject* Object : InObjects)
-	{
-		TArray<FMovieSceneObjectBindingID> ObjectBindings = Player->GetObjectBindings(Object);
-		for (FMovieSceneObjectBindingID ObjectBinding : ObjectBindings)
-		{
-			FMovieSceneSequenceID SequenceID = ObjectBinding.ResolveSequenceID(MovieSceneSequenceID::Root, *Player);
-			FMovieSceneBindingProxy Binding(ObjectBinding.GetGuid(), Player->State.FindSequence(SequenceID));
-			BoundObjects.Add(FSequencerBoundObjects(Binding, TArray<UObject*>({ Object })));
-		}
-	}
-
-	Player->Stop();
-	InWorld->DestroyActor(OutActor);
-	return BoundObjects;
-}
-
 void GatherDescendantBindings(const FMovieSceneBindingProxy& Binding, UMovieSceneSequence* Sequence, TArray<FMovieSceneBindingProxy>& AllBindings)
 {
 	UMovieScene* MovieScene = Sequence->GetMovieScene();
diff --git a/Engine/Plugins/MovieScene/SequencerScripting/Source/SequencerScriptingEditor/Public/SequencerTools.h b/Engine/Plugins/MovieScene/SequencerScripting/Source/SequencerScriptingEditor/Public/SequencerTools.h
--- a/Engine/Plugins/MovieScene/SequencerScripting/Source/SequencerScriptingEditor/Public/SequencerTools.h
+++ b/Engine/Plugins/MovieScene/SequencerScripting/Source/SequencerScriptingEditor/Public/SequencerTools.h
@@ -115,52 +115,6 @@ UCLASS(Transient, meta=(ScriptName="SequencerTools"))
 class SEQUENCERSCRIPTINGEDITOR_API USequencerToolsFunctionLibrary : public UBlueprintFunctionLibrary
 {
 	GENERATED_BODY()
-public:
-
-	/**
-	* Attempts to render a sequence to movie based on the specified settings. This will automatically detect
-	* if we're rendering via a PIE instance or a new process based on the passed in settings. Will return false
-	* if the state is not valid (ie: null or missing required parameters, capture in progress, etc.), true otherwise.
-	*/
-	UE_DEPRECATED(5.3, "SequencerTools::RenderMovie has been deprecated, please use Movie Render Queue")
-	UFUNCTION(BlueprintCallable, Category = "Editor Scripting | Sequencer Tools | Movie Rendering", meta=(DeprecatedFunction))
-	static bool RenderMovie(class UMovieSceneCapture* InCaptureSettings, FOnRenderMovieStopped OnFinishedCallback);
-
-	/** 
-	* Returns if Render to Movie is currently in progress.
-	*/
-	UE_DEPRECATED(5.3, "SequencerTools::IsRenderingMovie has been deprecated, please use Movie Render Queue")
-	UFUNCTION(BlueprintCallable, Category = "Editor Scripting | Sequencer Tools | Movie Rendering", meta=(DeprecatedFunction))
-	static bool IsRenderingMovie()
-	{
-		IMovieSceneCaptureDialogModule& MovieSceneCaptureModule = FModuleManager::Get().LoadModuleChecked<IMovieSceneCaptureDialogModule>("MovieSceneCaptureDialog");
-		return MovieSceneCaptureModule.GetCurrentCapture().IsValid();
-	}
-
-	/**
-	* Attempts to cancel an in-progress Render to Movie. Does nothing if there is no render in progress.
-	*/
-	UE_DEPRECATED(5.3, "SequencerTools::CancelMovieRender has been deprecated, please use Movie Render Queue")
-	UFUNCTION(BlueprintCallable, Category = "Editor Scripting | Sequencer Tools | Movie Rendering", meta=(DeprecatedFunction))
-	static void CancelMovieRender();
-
-public:
-
-	/*
-	 * Retrieve all objects currently bound to the specified binding identifiers. The sequence will be evaluated in lower bound of the specified range, 
-	 * which allows for retrieving spawnables in that period of time.
-	 */
-	UE_DEPRECATED(5.3, "SequencerTools::GetBoundObjects has been deprecated because it doesn't work with Spawnables, please use ULevelSequenceEditorBlueprintLibrary::GetBoundObjects which operates on the currently opened sequence in Editor")
-	UFUNCTION(BlueprintCallable, Category = "Editor Scripting | Sequencer Tools", meta=(DeprecatedFunction))
-	static TArray<FSequencerBoundObjects> GetBoundObjects(UWorld* InWorld, ULevelSequence* InSequence, const TArray<FMovieSceneBindingProxy>& InBindings, const FSequencerScriptingRange& InRange);
-
-	/*
-	 * Get the object bindings for the requested object. The sequence will be evaluated in lower bound of the specified range, 
-	 * which allows for retrieving spawnables in that period of time.
-	 */
-	UE_DEPRECATED(5.3, "SequencerTools::GetObjectBindings has been deprecated because it doesn't work with Spawnables, please use ULevelSequenceEditorBlueprintLibrary::GetObjectBindings which operates on the currently opened sequence in Editor")
-	UFUNCTION(BlueprintCallable, Category = "Editor Scripting | Sequencer Tools", meta=(DeprecatedFunction))
-	static TArray<FSequencerBoundObjects> GetObjectBindings(UWorld* InWorld, ULevelSequence* InSequence, const TArray<UObject*>& InObject, const FSequencerScriptingRange& InRange);
 
 public:
 
diff --git a/Engine/Plugins/MovieScene/TemplateSequence/Source/TemplateSequence/Private/CameraAnimationSequencePlayer.cpp b/Engine/Plugins/MovieScene/TemplateSequence/Source/TemplateSequence/Private/CameraAnimationSequencePlayer.cpp
--- a/Engine/Plugins/MovieScene/TemplateSequence/Source/TemplateSequence/Private/CameraAnimationSequencePlayer.cpp
+++ b/Engine/Plugins/MovieScene/TemplateSequence/Source/TemplateSequence/Private/CameraAnimationSequencePlayer.cpp
@@ -511,16 +511,3 @@ void UCameraAnimationSequencePlayer::Stop()
 		}
 	}
 }
-
-void UCameraAnimationSequencePlayer::StartScrubbing()
-{
-	ensure(Status == EMovieScenePlayerStatus::Playing);
-	Status = EMovieScenePlayerStatus::Scrubbing;
-}
-
-void UCameraAnimationSequencePlayer::EndScrubbing()
-{
-	ensure(Status == EMovieScenePlayerStatus::Scrubbing);
-	Status = EMovieScenePlayerStatus::Playing;
-}
-
diff --git a/Engine/Plugins/MovieScene/TemplateSequence/Source/TemplateSequence/Public/CameraAnimationSequencePlayer.h b/Engine/Plugins/MovieScene/TemplateSequence/Source/TemplateSequence/Public/CameraAnimationSequencePlayer.h
--- a/Engine/Plugins/MovieScene/TemplateSequence/Source/TemplateSequence/Public/CameraAnimationSequencePlayer.h
+++ b/Engine/Plugins/MovieScene/TemplateSequence/Source/TemplateSequence/Public/CameraAnimationSequencePlayer.h
@@ -152,14 +152,6 @@ public:
 	/** Stop playing the sequence */
 	void Stop();
 
-	/** Sets the player in scrub mode */
-	UE_DEPRECATED(5.3, "StartScrubbing has been deprecated as it has no functionality")
-	void StartScrubbing();
-
-	/** Ends scrub mode */
-	UE_DEPRECATED(5.3, "EndScrubbing has been deprecated as it has no functionality")
-	void EndScrubbing();
-
 	/** Gets whether playback is looping */
 	bool GetIsLooping() const { return bIsLooping; }
 	/** Get the sequence display rate */
diff --git a/Engine/Source/Editor/MovieSceneTools/Private/TrackEditors/SubTrackEditor.cpp b/Engine/Source/Editor/MovieSceneTools/Private/TrackEditors/SubTrackEditor.cpp
--- a/Engine/Source/Editor/MovieSceneTools/Private/TrackEditors/SubTrackEditor.cpp
+++ b/Engine/Source/Editor/MovieSceneTools/Private/TrackEditors/SubTrackEditor.cpp
@@ -844,11 +844,6 @@ void FSubTrackEditor::CreateNewTake(UMovieSceneSubSection* Section)
 	}
 }
 
-void FSubTrackEditor::SwitchTake(UObject* TakeObject)
-{
-	ChangeTake(Cast<UMovieSceneSequence>(TakeObject));
-}
-
 void FSubTrackEditor::ChangeTake(UMovieSceneSequence* Sequence)
 {
 	bool bChangedTake = false;
diff --git a/Engine/Source/Editor/MovieSceneTools/Public/TrackEditors/SubTrackEditor.h b/Engine/Source/Editor/MovieSceneTools/Public/TrackEditors/SubTrackEditor.h
--- a/Engine/Source/Editor/MovieSceneTools/Public/TrackEditors/SubTrackEditor.h
+++ b/Engine/Source/Editor/MovieSceneTools/Public/TrackEditors/SubTrackEditor.h
@@ -78,9 +78,6 @@ public:
 	/** Switch the selected section's take sequence */
 	virtual void ChangeTake(UMovieSceneSequence* Sequence);
 
-	UE_DEPRECATED(5.3, "SwitchTake has been deprecated, please use ChangeTake(UMovieSceneSequence*)")
-	virtual void SwitchTake(UObject* TakeObject);
-
 	/** Generate a menu for takes for this section */
 	virtual void AddTakesMenu(UMovieSceneSubSection* Section, FMenuBuilder& MenuBuilder);
 
diff --git a/Engine/Source/Runtime/Engine/Classes/Camera/CameraShakeBase.h b/Engine/Source/Runtime/Engine/Classes/Camera/CameraShakeBase.h
--- a/Engine/Source/Runtime/Engine/Classes/Camera/CameraShakeBase.h
+++ b/Engine/Source/Runtime/Engine/Classes/Camera/CameraShakeBase.h
@@ -383,14 +383,6 @@ struct FCameraShakeState
 	/** Helper method to get GetShakeInfo().Duration.IsInifnite() */
 	bool IsInfinite() const { return ShakeInfo.Duration.IsInfinite(); }
 
-public:
-
-	UE_DEPRECATED(5.3, "Please use Start")
-	void Initialize(const FCameraShakeInfo& InShakeInfo) { Start(InShakeInfo); }
-
-	UE_DEPRECATED(5.3, "Please use IsPlaying")
-	bool IsActive() const { return bIsPlaying; }
-
 private:
 
 	ENGINE_API void InitializePlaying();
@@ -593,9 +585,6 @@ public:
 	/** Tears down this camera shake before destruction or recycling */
 	ENGINE_API void TeardownShake();
 
-	UE_DEPRECATED(5.3, "Elapsed time doesn't exist anymore, get the information from the root shake pattern")
-	float GetElapsedTime() const { return 0.f; }
-
 public:
 
 	/** Gets the current camera manager. Will be null if the shake isn't active. */
diff --git a/Engine/Source/Runtime/LevelSequence/Public/LevelSequence.h b/Engine/Source/Runtime/LevelSequence/Public/LevelSequence.h
--- a/Engine/Source/Runtime/LevelSequence/Public/LevelSequence.h
+++ b/Engine/Source/Runtime/LevelSequence/Public/LevelSequence.h
@@ -100,9 +100,6 @@ public:
 
 	using UMovieSceneSequence::LocateBoundObjects;
 
-	UE_DEPRECATED(5.3, "Use LocateBoundObjects with FLocateBoundObjectsParams parameter instead")
-	void LocateBoundObjects(const FGuid& ObjectId, UObject* Context, const FTopLevelAssetPath& StreamedLevelAssetPath, TArray<UObject*, TInlineAllocator<1>>& OutObjects) const	{}
-
 	UE_DEPRECATED(5.4, "Use the base class LocateBoundObjects()")
 	LEVELSEQUENCE_API void LocateBoundObjects(const FGuid& ObjectId, UObject* Context, const FLevelSequenceBindingReference::FResolveBindingParams& InResolveBindingParams, TArray<UObject*, TInlineAllocator<1>>& OutObjects) const;
 
diff --git a/Engine/Source/Runtime/MovieScene/Private/MovieSceneSequenceTickManager.cpp b/Engine/Source/Runtime/MovieScene/Private/MovieSceneSequenceTickManager.cpp
--- a/Engine/Source/Runtime/MovieScene/Private/MovieSceneSequenceTickManager.cpp
+++ b/Engine/Source/Runtime/MovieScene/Private/MovieSceneSequenceTickManager.cpp
@@ -278,37 +278,6 @@ void UMovieSceneSequenceTickManager::UnregisterTickClientImpl(IMovieSceneSequenc
 	}
 }
 
-// DEPRECATED
-void UMovieSceneSequenceTickManager::RegisterSequenceActor(AActor* InActor)
-{
-	TScriptInterface<IMovieSceneSequenceTickManagerClient> Interface(InActor);
-	check(InActor && Interface);
-
-	FMovieSceneSequenceTickInterval ResolvedTickInterval(InActor);
-	RegisterTickClient(ResolvedTickInterval, Interface);
-}
-// DEPRECATED
-void UMovieSceneSequenceTickManager::RegisterSequenceActor(AActor* InActor, TScriptInterface<IMovieSceneSequenceTickManagerClient> InActorInterface)
-{
-	check(InActor && InActorInterface);
-
-	FMovieSceneSequenceTickInterval ResolvedTickInterval(InActor);
-	RegisterTickClient(ResolvedTickInterval, InActorInterface);
-}
-// DEPRECATED
-void UMovieSceneSequenceTickManager::UnregisterSequenceActor(AActor* InActor)
-{
-	TScriptInterface<IMovieSceneSequenceTickManagerClient> Interface(InActor);
-	check(Interface);
-	UnregisterTickClient(Interface);
-}
-// DEPRECATED
-void UMovieSceneSequenceTickManager::UnregisterSequenceActor(AActor* InActor, TScriptInterface<IMovieSceneSequenceTickManagerClient> InActorInterface)
-{
-	check(InActorInterface);
-	UnregisterTickClient(InActorInterface);
-}
-
 void UMovieSceneSequenceTickManager::TickSequenceActors(float DeltaSeconds)
 {
 	CSV_SCOPED_TIMING_STAT_EXCLUSIVE(MovieSceneEval);
diff --git a/Engine/Source/Runtime/MovieScene/Public/EntitySystem/MovieSceneEntitySystemLinker.h b/Engine/Source/Runtime/MovieScene/Public/EntitySystem/MovieSceneEntitySystemLinker.h
--- a/Engine/Source/Runtime/MovieScene/Public/EntitySystem/MovieSceneEntitySystemLinker.h
+++ b/Engine/Source/Runtime/MovieScene/Public/EntitySystem/MovieSceneEntitySystemLinker.h
@@ -389,9 +389,6 @@ public:
 		FMovieSceneEntitySystemLinkerEvent          CleanTaggedGarbage;
 		FMovieSceneEntitySystemLinkerAROEvent       AddReferencedObjects;
 		FMovieSceneEntitySystemLinkerEvent          AbandonLinker;
-
-		UE_DEPRECATED(5.3, "Please use FWorldDelegates::OnWorldCleanup directly")
-		FMovieSceneEntitySystemLinkerWorldEvent     CleanUpWorld;
 	} Events;
 
 private:
diff --git a/Engine/Source/Runtime/MovieScene/Public/Evaluation/MovieSceneEvaluationTemplateInstance.h b/Engine/Source/Runtime/MovieScene/Public/Evaluation/MovieSceneEvaluationTemplateInstance.h
--- a/Engine/Source/Runtime/MovieScene/Public/Evaluation/MovieSceneEvaluationTemplateInstance.h
+++ b/Engine/Source/Runtime/MovieScene/Public/Evaluation/MovieSceneEvaluationTemplateInstance.h
@@ -109,12 +109,6 @@ public:
 
 	MOVIESCENE_API void CopyActuators(FMovieSceneBlendingAccumulator& Accumulator) const;
 
-	UE_DEPRECATED(5.3, "Please call TearDown instead.")
-	void BeginDestroy()
-	{
-		TearDown();
-	}
-
 	MOVIESCENE_API void TearDown();
 
 	MOVIESCENE_API void EnableGlobalPreAnimatedStateCapture();
diff --git a/Engine/Source/Runtime/MovieScene/Public/MovieSceneSequenceTickManager.h b/Engine/Source/Runtime/MovieScene/Public/MovieSceneSequenceTickManager.h
--- a/Engine/Source/Runtime/MovieScene/Public/MovieSceneSequenceTickManager.h
+++ b/Engine/Source/Runtime/MovieScene/Public/MovieSceneSequenceTickManager.h
@@ -86,18 +86,6 @@ public:
 
 public:
 
-	/*~ -------------------------------------------------------------- */
-	/*~ Deprecated functions that previously used 'Actor' nomenclature */
-	UE_DEPRECATED(5.1, "Please use RegisterTickClient")
-	MOVIESCENE_API void RegisterSequenceActor(AActor* InActor);
-	UE_DEPRECATED(5.1, "Please use RegisterTickClient")
-	MOVIESCENE_API void RegisterSequenceActor(AActor* InActor, TScriptInterface<IMovieSceneSequenceTickManagerClient> InActorInterface);
-	UE_DEPRECATED(5.1, "Please use UnregisterTickClient")
-	MOVIESCENE_API void UnregisterSequenceActor(AActor* InActor);
-	UE_DEPRECATED(5.1, "Please use UnregisterTickClient")
-	MOVIESCENE_API void UnregisterSequenceActor(AActor* InActor, TScriptInterface<IMovieSceneSequenceTickManagerClient> InActorInterface);
-	/*~ -------------------------------------------------------------- */
-
 	static MOVIESCENE_API void AddReferencedObjects(UObject* InThis, FReferenceCollector& Collector);
 
 private:
diff --git a/Engine/Source/Runtime/MovieScene/Public/MovieSceneSequenceTickManagerClient.h b/Engine/Source/Runtime/MovieScene/Public/MovieSceneSequenceTickManagerClient.h
--- a/Engine/Source/Runtime/MovieScene/Public/MovieSceneSequenceTickManagerClient.h
+++ b/Engine/Source/Runtime/MovieScene/Public/MovieSceneSequenceTickManagerClient.h
@@ -28,14 +28,3 @@ public:
 
 	virtual void TickFromSequenceTickManager(float DeltaSeconds, FMovieSceneEntitySystemRunner* Runner) = 0;
 };
-
-/** Deprecated IMovieSceneSequenceActor interface that has been replaced with IMovieSceneSequenceTickManagerClient */
-class UE_DEPRECATED(5.1, "Please use IMovieSceneSequenceTickManagerClient directly") IMovieSceneSequenceActor
-	: public IMovieSceneSequenceTickManagerClient
-{
-	virtual void TickFromSequenceTickManager(float DeltaSeconds) = 0;
-	virtual void TickFromSequenceTickManager(float DeltaSeconds, FMovieSceneEntitySystemRunner* Runner) override
-	{
-		TickFromSequenceTickManager(DeltaSeconds);
-	}
-};
