Sequencer: Put RenderMovie and GetBoundObjects back for now

#rnx

diff --git a/Engine/Plugins/MovieScene/SequencerScripting/Source/SequencerScriptingEditor/Private/SequencerTools.cpp b/Engine/Plugins/MovieScene/SequencerScripting/Source/SequencerScriptingEditor/Private/SequencerTools.cpp
--- a/Engine/Plugins/MovieScene/SequencerScripting/Source/SequencerScriptingEditor/Private/SequencerTools.cpp
+++ b/Engine/Plugins/MovieScene/SequencerScripting/Source/SequencerScriptingEditor/Private/SequencerTools.cpp
@@ -40,6 +40,147 @@
 
 #define LOCTEXT_NAMESPACE "SequencerTools"
 
+bool USequencerToolsFunctionLibrary::RenderMovie(UMovieSceneCapture* InCaptureSettings, FOnRenderMovieStopped OnFinishedCallback)
+{
+	IMovieSceneCaptureDialogModule& MovieSceneCaptureModule = FModuleManager::Get().LoadModuleChecked<IMovieSceneCaptureDialogModule>("MovieSceneCaptureDialog");
+	
+	// Because this comes from the Python/BP layer we need to soft-validate the state before we pass it onto functions that do a assert-based validation.
+	if (!InCaptureSettings)
+	{
+		FFrame::KismetExecutionMessage(TEXT("Cannot start Render Sequence to Movie with null capture settings."), ELogVerbosity::Error);
+		return false;
+	}
+
+PRAGMA_DISABLE_DEPRECATION_WARNINGS
+	if (IsRenderingMovie())
+	{
+		FFrame::KismetExecutionMessage(TEXT("Capture already in progress."), ELogVerbosity::Error);
+		return false;
+	}
+PRAGMA_ENABLE_DEPRECATION_WARNINGS
+
+	// If they're capturing a level sequence we'll do some additional checking as there are more parameters on the Automated Level Sequence capture.
+	UAutomatedLevelSequenceCapture* LevelSequenceCapture = Cast<UAutomatedLevelSequenceCapture>(InCaptureSettings);
+	if (LevelSequenceCapture)
+	{
+		if (!LevelSequenceCapture->LevelSequenceAsset.IsValid())
+		{
+			// UE_LOG(LogTemp, Warning, TEXT("No Level Sequence Asset specified in UAutomatedLevelSequenceCapture."));
+			FFrame::KismetExecutionMessage(TEXT("No Level Sequence Asset specified in UAutomatedLevelSequenceCapture."), ELogVerbosity::Error);
+			return false;
+		}
+
+		if (!LevelSequenceCapture->bUseCustomStartFrame && !LevelSequenceCapture->bUseCustomEndFrame)
+		{
+			// If they don't want to use a custom start/end frame we override the default values to be the length of the sequence, as the default is [0,1)
+			ULevelSequence* LevelSequence = Cast<ULevelSequence>(LevelSequenceCapture->LevelSequenceAsset.TryLoad());
+			if (!LevelSequence)
+			{
+				const FString ErrorMessage = FString::Printf(TEXT("Specified Level Sequence Asset failed to load. Specified Asset Path: %s"), *LevelSequenceCapture->LevelSequenceAsset.GetAssetPathString());
+				FFrame::KismetExecutionMessage(*ErrorMessage, ELogVerbosity::Error);
+				return false;
+			}
+
+			FFrameRate DisplayRate = LevelSequence->GetMovieScene()->GetDisplayRate();
+			FFrameRate TickResolution = LevelSequence->GetMovieScene()->GetTickResolution();
+			
+			LevelSequenceCapture->Settings.CustomFrameRate = DisplayRate;
+			LevelSequenceCapture->Settings.bUseCustomFrameRate = true;
+			LevelSequenceCapture->Settings.bUseRelativeFrameNumbers = false;
+			TRange<FFrameNumber> Range = LevelSequence->GetMovieScene()->GetPlaybackRange();
+
+			FFrameNumber StartFrame = UE::MovieScene::DiscreteInclusiveLower(Range);
+			FFrameNumber EndFrame = UE::MovieScene::DiscreteExclusiveUpper(Range);
+
+			FFrameNumber RoundedStartFrame = FFrameRate::TransformTime(StartFrame, TickResolution, DisplayRate).CeilToFrame();
+			FFrameNumber RoundedEndFrame = FFrameRate::TransformTime(EndFrame, TickResolution, DisplayRate).CeilToFrame();
+
+			LevelSequenceCapture->CustomStartFrame = RoundedStartFrame;
+			LevelSequenceCapture->CustomEndFrame = RoundedEndFrame;
+		}
+	}
+
+	auto LocalCaptureStoppedCallback = [OnFinishedCallback](bool bSuccess)
+	{
+		OnFinishedCallback.ExecuteIfBound(bSuccess);
+	};
+
+	MovieSceneCaptureModule.StartCapture(InCaptureSettings);
+	MovieSceneCaptureModule.GetCurrentCapture()->CaptureStoppedDelegate.AddLambda(LocalCaptureStoppedCallback);
+	return true;
+}
+
+void USequencerToolsFunctionLibrary::CancelMovieRender()
+{
+	IMovieSceneCaptureDialogModule& MovieSceneCaptureModule = FModuleManager::Get().LoadModuleChecked<IMovieSceneCaptureDialogModule>("MovieSceneCaptureDialog");
+	TSharedPtr<FMovieSceneCaptureBase> CurrentCapture = MovieSceneCaptureModule.GetCurrentCapture();
+	if (CurrentCapture.IsValid())
+	{
+		// We just invoke the capture's Cancel function. This will cause a shut-down of the capture (the same as the UI)
+		// which will invoke all of the necessary callbacks as well. We don't null out CurrentCapture because that is done
+		// as the result of its shutdown callbacks.
+		CurrentCapture->Cancel();
+	}
+}
+
+TArray<FSequencerBoundObjects> USequencerToolsFunctionLibrary::GetBoundObjects(UWorld* InWorld, ULevelSequence* InSequence, const TArray<FMovieSceneBindingProxy>& InBindings, const FSequencerScriptingRange& InRange)
+{
+	ALevelSequenceActor* OutActor;
+	FMovieSceneSequencePlaybackSettings Settings;
+	FLevelSequenceCameraSettings CameraSettings;
+
+	ULevelSequencePlayer* Player = ULevelSequencePlayer::CreateLevelSequencePlayer(InWorld, InSequence, Settings, OutActor);
+
+	// Evaluation needs to occur in order to obtain spawnables
+	FFrameRate Resolution = InSequence->GetMovieScene()->GetTickResolution();
+	TRange<FFrameNumber> SpecifiedRange = InRange.ToNative(Resolution);
+	Player->SetPlaybackPosition(FMovieSceneSequencePlaybackParams(SpecifiedRange.GetLowerBoundValue().Value, EUpdatePositionMethod::Play));
+
+	FMovieSceneSequenceID SequenceId = Player->State.FindSequenceId(InSequence);
+
+	TArray<FSequencerBoundObjects> BoundObjects;
+	for (FMovieSceneBindingProxy Binding : InBindings)
+	{
+		FMovieSceneObjectBindingID ObjectBinding = UE::MovieScene::FFixedObjectBindingID(Binding.BindingID, SequenceId);
+		BoundObjects.Add(FSequencerBoundObjects(Binding, Player->GetBoundObjects(ObjectBinding)));
+	}
+
+	Player->Stop();
+	InWorld->DestroyActor(OutActor);
+	return BoundObjects;
+}
+
+
+TArray<FSequencerBoundObjects> USequencerToolsFunctionLibrary::GetObjectBindings(UWorld* InWorld, ULevelSequence* InSequence, const TArray<UObject*>& InObjects, const FSequencerScriptingRange& InRange)
+{
+	ALevelSequenceActor* OutActor;
+	FMovieSceneSequencePlaybackSettings Settings;
+	FLevelSequenceCameraSettings CameraSettings;
+
+	ULevelSequencePlayer* Player = ULevelSequencePlayer::CreateLevelSequencePlayer(InWorld, InSequence, Settings, OutActor);
+
+	FFrameRate Resolution = InSequence->GetMovieScene()->GetTickResolution();
+	TRange<FFrameNumber> SpecifiedRange = InRange.ToNative(Resolution);
+	Player->SetPlaybackPosition(FMovieSceneSequencePlaybackParams(SpecifiedRange.GetLowerBoundValue().Value, EUpdatePositionMethod::Play));
+
+	TArray<FSequencerBoundObjects> BoundObjects;
+
+	for (UObject* Object : InObjects)
+	{
+		TArray<FMovieSceneObjectBindingID> ObjectBindings = Player->GetObjectBindings(Object);
+		for (FMovieSceneObjectBindingID ObjectBinding : ObjectBindings)
+		{
+			FMovieSceneSequenceID SequenceID = ObjectBinding.ResolveSequenceID(MovieSceneSequenceID::Root, *Player);
+			FMovieSceneBindingProxy Binding(ObjectBinding.GetGuid(), Player->State.FindSequence(SequenceID));
+			BoundObjects.Add(FSequencerBoundObjects(Binding, TArray<UObject*>({ Object })));
+		}
+	}
+
+	Player->Stop();
+	InWorld->DestroyActor(OutActor);
+	return BoundObjects;
+}
+
 void GatherDescendantBindings(const FMovieSceneBindingProxy& Binding, UMovieSceneSequence* Sequence, TArray<FMovieSceneBindingProxy>& AllBindings)
 {
 	UMovieScene* MovieScene = Sequence->GetMovieScene();
diff --git a/Engine/Plugins/MovieScene/SequencerScripting/Source/SequencerScriptingEditor/Public/SequencerTools.h b/Engine/Plugins/MovieScene/SequencerScripting/Source/SequencerScriptingEditor/Public/SequencerTools.h
--- a/Engine/Plugins/MovieScene/SequencerScripting/Source/SequencerScriptingEditor/Public/SequencerTools.h
+++ b/Engine/Plugins/MovieScene/SequencerScripting/Source/SequencerScriptingEditor/Public/SequencerTools.h
@@ -115,6 +115,52 @@ UCLASS(Transient, meta=(ScriptName="SequencerTools"))
 class SEQUENCERSCRIPTINGEDITOR_API USequencerToolsFunctionLibrary : public UBlueprintFunctionLibrary
 {
 	GENERATED_BODY()
+public:
+
+	/**
+	* Attempts to render a sequence to movie based on the specified settings. This will automatically detect
+	* if we're rendering via a PIE instance or a new process based on the passed in settings. Will return false
+	* if the state is not valid (ie: null or missing required parameters, capture in progress, etc.), true otherwise.
+	*/
+	UE_DEPRECATED(5.3, "SequencerTools::RenderMovie has been deprecated, please use Movie Render Queue")
+	UFUNCTION(BlueprintCallable, Category = "Editor Scripting | Sequencer Tools | Movie Rendering", meta=(DeprecatedFunction))
+	static bool RenderMovie(class UMovieSceneCapture* InCaptureSettings, FOnRenderMovieStopped OnFinishedCallback);
+
+	/** 
+	* Returns if Render to Movie is currently in progress.
+	*/
+	UE_DEPRECATED(5.3, "SequencerTools::IsRenderingMovie has been deprecated, please use Movie Render Queue")
+	UFUNCTION(BlueprintCallable, Category = "Editor Scripting | Sequencer Tools | Movie Rendering", meta=(DeprecatedFunction))
+	static bool IsRenderingMovie()
+	{
+		IMovieSceneCaptureDialogModule& MovieSceneCaptureModule = FModuleManager::Get().LoadModuleChecked<IMovieSceneCaptureDialogModule>("MovieSceneCaptureDialog");
+		return MovieSceneCaptureModule.GetCurrentCapture().IsValid();
+	}
+
+	/**
+	* Attempts to cancel an in-progress Render to Movie. Does nothing if there is no render in progress.
+	*/
+	UE_DEPRECATED(5.3, "SequencerTools::CancelMovieRender has been deprecated, please use Movie Render Queue")
+	UFUNCTION(BlueprintCallable, Category = "Editor Scripting | Sequencer Tools | Movie Rendering", meta=(DeprecatedFunction))
+	static void CancelMovieRender();
+
+public:
+
+	/*
+	 * Retrieve all objects currently bound to the specified binding identifiers. The sequence will be evaluated in lower bound of the specified range, 
+	 * which allows for retrieving spawnables in that period of time.
+	 */
+	UE_DEPRECATED(5.3, "SequencerTools::GetBoundObjects has been deprecated because it doesn't work with Spawnables, please use ULevelSequenceEditorBlueprintLibrary::GetBoundObjects which operates on the currently opened sequence in Editor")
+	UFUNCTION(BlueprintCallable, Category = "Editor Scripting | Sequencer Tools", meta=(DeprecatedFunction))
+	static TArray<FSequencerBoundObjects> GetBoundObjects(UWorld* InWorld, ULevelSequence* InSequence, const TArray<FMovieSceneBindingProxy>& InBindings, const FSequencerScriptingRange& InRange);
+
+	/*
+	 * Get the object bindings for the requested object. The sequence will be evaluated in lower bound of the specified range, 
+	 * which allows for retrieving spawnables in that period of time.
+	 */
+	UE_DEPRECATED(5.3, "SequencerTools::GetObjectBindings has been deprecated because it doesn't work with Spawnables, please use ULevelSequenceEditorBlueprintLibrary::GetObjectBindings which operates on the currently opened sequence in Editor")
+	UFUNCTION(BlueprintCallable, Category = "Editor Scripting | Sequencer Tools", meta=(DeprecatedFunction))
+	static TArray<FSequencerBoundObjects> GetObjectBindings(UWorld* InWorld, ULevelSequence* InSequence, const TArray<UObject*>& InObject, const FSequencerScriptingRange& InRange);
 
 public:
 
