Sequencer: Add conversion functions for BP and python to FMovieSceneTimeWarpVariant

diff --git a/Engine/Plugins/MovieScene/SequencerScripting/Source/SequencerScripting/Private/ExtensionLibraries/MovieSceneTimeWarpExtensions.cpp b/Engine/Plugins/MovieScene/SequencerScripting/Source/SequencerScripting/Private/ExtensionLibraries/MovieSceneTimeWarpExtensions.cpp
new file mode 100644
--- /dev/null
+++ b/Engine/Plugins/MovieScene/SequencerScripting/Source/SequencerScripting/Private/ExtensionLibraries/MovieSceneTimeWarpExtensions.cpp
@@ -0,0 +1,47 @@
+// Copyright Epic Games, Inc. All Rights Reserved.
+
+#include "ExtensionLibraries/MovieSceneTimeWarpExtensions.h"
+#include "Variants/MovieSceneTimeWarpVariant.h"
+#include "UObject/Stack.h"
+
+#include UE_INLINE_GENERATED_CPP_BY_NAME(MovieSceneTimeWarpExtensions)
+
+
+double UMovieSceneTimeWarpExtensions::Conv_TimeWarpVariantToPlayRate(const FMovieSceneTimeWarpVariant& TimeWarp)
+{
+	if (TimeWarp.GetType() == EMovieSceneTimeWarpType::FixedPlayRate)
+	{
+		return TimeWarp.AsFixedPlayRate();
+	}
+
+	FFrame::KismetExecutionMessage(TEXT("The specified TimeWarp is not a fixed play rate"), ELogVerbosity::Error);
+	return 1.f;
+}
+
+FMovieSceneTimeWarpVariant UMovieSceneTimeWarpExtensions::Conv_PlayRateToTimeWarpVariant(double ConstantPlayRate)
+{
+	return FMovieSceneTimeWarpVariant(ConstantPlayRate);
+}
+
+double UMovieSceneTimeWarpExtensions::ToFixedPlayRate(const FMovieSceneTimeWarpVariant& TimeWarp)
+{
+	return Conv_TimeWarpVariantToPlayRate(TimeWarp);
+}
+
+void UMovieSceneTimeWarpExtensions::SetFixedPlayRate(FMovieSceneTimeWarpVariant& TimeWarp, double FixedPlayRate)
+{
+	TimeWarp.Set(FixedPlayRate);
+}
+
+void UMovieSceneTimeWarpExtensions::BreakTimeWarp(const FMovieSceneTimeWarpVariant& TimeWarp, double& FixedPlayRate)
+{
+	if (TimeWarp.GetType() == EMovieSceneTimeWarpType::FixedPlayRate)
+	{
+		FixedPlayRate = TimeWarp.AsFixedPlayRate();
+	}
+}
+
+FMovieSceneTimeWarpVariant UMovieSceneTimeWarpExtensions::MakeTimeWarp(double FixedPlayRate)
+{
+	return FMovieSceneTimeWarpVariant(FixedPlayRate);
+}
\ No newline at end of file
diff --git a/Engine/Plugins/MovieScene/SequencerScripting/Source/SequencerScripting/Public/ExtensionLibraries/MovieSceneTimeWarpExtensions.h b/Engine/Plugins/MovieScene/SequencerScripting/Source/SequencerScripting/Public/ExtensionLibraries/MovieSceneTimeWarpExtensions.h
new file mode 100644
--- /dev/null
+++ b/Engine/Plugins/MovieScene/SequencerScripting/Source/SequencerScripting/Public/ExtensionLibraries/MovieSceneTimeWarpExtensions.h
@@ -0,0 +1,52 @@
+// Copyright Epic Games, Inc. All Rights Reserved.
+
+#pragma once
+
+#include "Kismet/BlueprintFunctionLibrary.h"
+#include "MovieSceneTimeWarpExtensions.generated.h"
+
+struct FMovieSceneTimeWarpVariant;
+
+
+/**
+ * Function library containing methods that relate to time-warp within Sequencer
+ */
+UCLASS()
+class UMovieSceneTimeWarpExtensions : public UBlueprintFunctionLibrary
+{
+	GENERATED_BODY()
+
+public:
+
+	/**
+	 * Converts a timewarp variant struct to a constant play rate
+	 */
+	UFUNCTION(BlueprintPure, meta=(DisplayName = "To PlayRate", BlueprintAutocast), Category="Sequencer|TimeWarp")
+	static SEQUENCERSCRIPTING_API double Conv_TimeWarpVariantToPlayRate(const FMovieSceneTimeWarpVariant& TimeWarp);
+
+	/**
+	 * Converts a constant playrate to a timewarp variant
+	 */
+	UFUNCTION(BlueprintPure, meta=(DisplayName = "To TimeWarp", BlueprintAutocast), Category="Sequencer|TimeWarp")
+	static SEQUENCERSCRIPTING_API FMovieSceneTimeWarpVariant Conv_PlayRateToTimeWarpVariant(double ConstantPlayRate);
+
+	/**
+	 * Retrieve this timewarp's constant play rate. Will throw an error if the timewarp is not a constant play rate.
+	 */
+	UFUNCTION(BlueprintPure, meta=(ScriptMethod), Category="Sequencer|TimeWarp")
+	static SEQUENCERSCRIPTING_API double ToFixedPlayRate(const FMovieSceneTimeWarpVariant& TimeWarp);
+
+	/**
+	 * Assign a constant playrate to this timewarp, overwriting any existing timewarp implementation.
+	 */
+	UFUNCTION(BlueprintPure, meta=(ScriptMethod), Category="Sequencer|TimeWarp")
+	static SEQUENCERSCRIPTING_API void SetFixedPlayRate(UPARAM(ref) FMovieSceneTimeWarpVariant& TimeWarp, double FixedPlayRate);
+
+public:
+
+	UFUNCTION(BlueprintPure, Category="Sequencer|TimeWarp")
+	static SEQUENCERSCRIPTING_API void BreakTimeWarp(const FMovieSceneTimeWarpVariant& TimeWarp, double& FixedPlayRate);
+
+	UFUNCTION(BlueprintPure, Category="Sequencer|TimeWarp")
+	static SEQUENCERSCRIPTING_API FMovieSceneTimeWarpVariant MakeTimeWarp(double FixedPlayRate);
+};
\ No newline at end of file
diff --git a/Engine/Source/Runtime/MovieScene/Public/Variants/MovieSceneTimeWarpVariant.h b/Engine/Source/Runtime/MovieScene/Public/Variants/MovieSceneTimeWarpVariant.h
--- a/Engine/Source/Runtime/MovieScene/Public/Variants/MovieSceneTimeWarpVariant.h
+++ b/Engine/Source/Runtime/MovieScene/Public/Variants/MovieSceneTimeWarpVariant.h
@@ -46,7 +46,7 @@ enum class EMovieSceneTimeWarpType : uint8
  * By default this variant is a literal value that represents a play rate of 1.0 (ie, a 1:1 mapping), but it can be customized
  * to provide a wide range of different transformations such as looping, clamping and custom curves
  */
-USTRUCT(BlueprintType)
+USTRUCT(BlueprintType, meta=(HasNativeBreak="/Script/SequencerScripting.MovieSceneTimeWarpExtensions.BreakTimeWarp", HasNativeMake="/Script/SequencerScripting.MovieSceneTimeWarpExtensions.MakeTimeWarp"))
 struct FMovieSceneTimeWarpVariant
 {
 	GENERATED_BODY()
