Define scripting comparision (== and !=) for MovieScene Binding Proxies
This lets users compare their binding proxies directly in scripts.

diff --git a/Engine/Plugins/MovieScene/SequencerScripting/Source/SequencerScripting/Private/ExtensionLibraries/MovieSceneBindingExtensions.cpp b/Engine/Plugins/MovieScene/SequencerScripting/Source/SequencerScripting/Private/ExtensionLibraries/MovieSceneBindingExtensions.cpp
--- a/Engine/Plugins/MovieScene/SequencerScripting/Source/SequencerScripting/Private/ExtensionLibraries/MovieSceneBindingExtensions.cpp
+++ b/Engine/Plugins/MovieScene/SequencerScripting/Source/SequencerScripting/Private/ExtensionLibraries/MovieSceneBindingExtensions.cpp
@@ -28,6 +28,16 @@ bool UMovieSceneBindingExtensions::IsValid(const FMovieSceneBindingProxy& InBind
 	return false;
 }
 
+bool UMovieSceneBindingExtensions::EqualEqual_MovieSceneBindingProxy(FMovieSceneBindingProxy A, FMovieSceneBindingProxy B)
+{
+	return A == B;
+}
+
+bool UMovieSceneBindingExtensions::NotEqual_MovieSceneBindingProxy(FMovieSceneBindingProxy A, FMovieSceneBindingProxy B)
+{
+	return A != B;
+}
+
 FGuid UMovieSceneBindingExtensions::GetId(const FMovieSceneBindingProxy& InBinding)
 {
 	return InBinding.BindingID;
diff --git a/Engine/Plugins/MovieScene/SequencerScripting/Source/SequencerScripting/Public/ExtensionLibraries/MovieSceneBindingExtensions.h b/Engine/Plugins/MovieScene/SequencerScripting/Source/SequencerScripting/Public/ExtensionLibraries/MovieSceneBindingExtensions.h
--- a/Engine/Plugins/MovieScene/SequencerScripting/Source/SequencerScripting/Public/ExtensionLibraries/MovieSceneBindingExtensions.h
+++ b/Engine/Plugins/MovieScene/SequencerScripting/Source/SequencerScripting/Public/ExtensionLibraries/MovieSceneBindingExtensions.h
@@ -25,6 +25,14 @@ public:
 	UFUNCTION(BlueprintPure, Category = "Sequencer|Sequence", meta=(ScriptMethod))
 	static bool IsValid(const FMovieSceneBindingProxy& InBinding);
 
+	/** Returns true if the values are equal (A == B) */
+	UFUNCTION(BlueprintPure, meta = (DisplayName = "Equal (MovieSceneBindingProxy)", CompactNodeTitle = "==", ScriptOperator = "==", BlueprintThreadSafe), Category = "Sequencer|Sequence")
+	static bool EqualEqual_MovieSceneBindingProxy(FMovieSceneBindingProxy A, FMovieSceneBindingProxy B);
+
+	/** Returns true if the values are not equal (A != B) */
+	UFUNCTION(BlueprintPure, meta = (DisplayName = "Not Equal (MovieSceneBindingProxy)", CompactNodeTitle = "!=", ScriptOperator = "!=", BlueprintThreadSafe), Category = "Sequencer|Sequence")
+	static bool NotEqual_MovieSceneBindingProxy(FMovieSceneBindingProxy A, FMovieSceneBindingProxy B);
+
 	/**
 	 * Get this binding's ID
 	 *
