Sequencer: Add modify for SetObjectPropertyClass

diff --git a/Engine/Plugins/MovieScene/SequencerScripting/Source/SequencerScripting/Private/ExtensionLibraries/MovieScenePropertyTrackExtensions.cpp b/Engine/Plugins/MovieScene/SequencerScripting/Source/SequencerScripting/Private/ExtensionLibraries/MovieScenePropertyTrackExtensions.cpp
--- a/Engine/Plugins/MovieScene/SequencerScripting/Source/SequencerScripting/Private/ExtensionLibraries/MovieScenePropertyTrackExtensions.cpp
+++ b/Engine/Plugins/MovieScene/SequencerScripting/Source/SequencerScripting/Private/ExtensionLibraries/MovieScenePropertyTrackExtensions.cpp
@@ -63,6 +63,7 @@ void UMovieScenePropertyTrackExtensions::SetObjectPropertyClass(UMovieSceneObjec
 		return;
 	}
 
+	Track->Modify();
 	Track->PropertyClass = PropertyClass;
 }
 
@@ -85,6 +86,7 @@ void UMovieScenePropertyTrackExtensions::SetByteTrackEnum(UMovieSceneByteTrack*
 		return;
 	}
 
+	Track->Modify();
 	Track->SetEnum(InEnum);
 }
 
