Sequencer Scripting: Added transform to channel manipulation which allows you to offset and scale the keys in a range

This uses an IRetimingInterface instead of UMovieSceneChannel::Offset

diff --git a/Engine/Plugins/MovieScene/SequencerScripting/Source/SequencerScripting/Private/KeysAndChannels/MovieSceneScriptingActorReference.h b/Engine/Plugins/MovieScene/SequencerScripting/Source/SequencerScripting/Private/KeysAndChannels/MovieSceneScriptingActorReference.h
--- a/Engine/Plugins/MovieScene/SequencerScripting/Source/SequencerScripting/Private/KeysAndChannels/MovieSceneScriptingActorReference.h
+++ b/Engine/Plugins/MovieScene/SequencerScripting/Source/SequencerScripting/Private/KeysAndChannels/MovieSceneScriptingActorReference.h
@@ -185,6 +185,22 @@ public:
 	{
 		return GetDefault() == FMovieSceneObjectBindingID();
 	}
+
+	/**
+	 * Transform the keys in time in the channel by an offset, scale and pivot
+	 * @param	OffsetFrame	    The amount to offset the keys by
+	 * @param	Scale	        The amount to scale the keys by
+	 * @param	PivotFrame	    The frame to pivot around when scaling the keys
+	 * @param   ScriptingRange  The time range of the keys to transform
+	 * @param	TimeUnit 		Is the specified OffsetFrame, PivotFrame, and ScriptingRange in Display Rate frames or Tick Resolution.
+	 *
+	 */
+	UFUNCTION(BlueprintCallable, Category = "Sequencer|Keys", meta = (DisplayName = "Transform (Actor Reference)"))
+	void Transform(FFrameNumber OffsetFrame, double Scale, FFrameNumber PivotFrame, FSequencerScriptingRange ScriptingRange, EMovieSceneTimeUnit TimeUnit = EMovieSceneTimeUnit::DisplayRate)
+	{
+		Impl::TransformKeysInChannel(ChannelHandle, OwningSequence, OwningSection, OffsetFrame, Scale, PivotFrame, ScriptingRange, TimeUnit);
+	}
+
 public:
 	TWeakObjectPtr<UMovieSceneSequence> OwningSequence;
 	TWeakObjectPtr<UMovieSceneSection> OwningSection;
diff --git a/Engine/Plugins/MovieScene/SequencerScripting/Source/SequencerScripting/Private/KeysAndChannels/MovieSceneScriptingBool.h b/Engine/Plugins/MovieScene/SequencerScripting/Source/SequencerScripting/Private/KeysAndChannels/MovieSceneScriptingBool.h
--- a/Engine/Plugins/MovieScene/SequencerScripting/Source/SequencerScripting/Private/KeysAndChannels/MovieSceneScriptingBool.h
+++ b/Engine/Plugins/MovieScene/SequencerScripting/Source/SequencerScripting/Private/KeysAndChannels/MovieSceneScriptingBool.h
@@ -182,6 +182,22 @@ public:
 	{
 		return Impl::GetDefaultFromChannel(ChannelHandle).IsSet();
 	}
+
+	/**
+	 * Transform the keys in time in the channel by an offset, scale and pivot
+	 * @param	OffsetFrame	    The amount to offset the keys by
+	 * @param	Scale	        The amount to scale the keys by
+	 * @param	PivotFrame	    The frame to pivot around when scaling the keys
+	 * @param   ScriptingRange  The time range of the keys to transform
+	 * @param	TimeUnit 		Is the specified OffsetFrame, PivotFrame, and ScriptingRange in Display Rate frames or Tick Resolution.
+	 *
+	 */
+	UFUNCTION(BlueprintCallable, Category = "Sequencer|Keys", meta = (DisplayName = "Transform (Bool)"))
+	void Transform(FFrameNumber OffsetFrame, double Scale, FFrameNumber PivotFrame, FSequencerScriptingRange ScriptingRange, EMovieSceneTimeUnit TimeUnit = EMovieSceneTimeUnit::DisplayRate)
+	{
+		Impl::TransformKeysInChannel(ChannelHandle, OwningSequence, OwningSection, OffsetFrame, Scale, PivotFrame, ScriptingRange, TimeUnit);
+	}
+
 public:
 	TWeakObjectPtr<UMovieSceneSequence> OwningSequence;
 	TWeakObjectPtr<UMovieSceneSection> OwningSection;
diff --git a/Engine/Plugins/MovieScene/SequencerScripting/Source/SequencerScripting/Private/KeysAndChannels/MovieSceneScriptingByte.h b/Engine/Plugins/MovieScene/SequencerScripting/Source/SequencerScripting/Private/KeysAndChannels/MovieSceneScriptingByte.h
--- a/Engine/Plugins/MovieScene/SequencerScripting/Source/SequencerScripting/Private/KeysAndChannels/MovieSceneScriptingByte.h
+++ b/Engine/Plugins/MovieScene/SequencerScripting/Source/SequencerScripting/Private/KeysAndChannels/MovieSceneScriptingByte.h
@@ -154,6 +154,22 @@ public:
 	{
 		return Impl::GetDefaultFromChannel(ChannelHandle).IsSet();
 	}
+
+	/**
+	 * Transform the keys in time in the channel by an offset, scale and pivot
+	 * @param	OffsetFrame	    The amount to offset the keys by
+	 * @param	Scale	        The amount to scale the keys by
+	 * @param	PivotFrame	    The frame to pivot around when scaling the keys
+	 * @param   ScriptingRange  The time range of the keys to transform
+	 * @param	TimeUnit 		Is the specified OffsetFrame, PivotFrame, and ScriptingRange in Display Rate frames or Tick Resolution.
+	 *
+	 */
+	UFUNCTION(BlueprintCallable, Category = "Sequencer|Keys", meta = (DisplayName = "Transform (Enum)"))
+	void Transform(FFrameNumber OffsetFrame, double Scale, FFrameNumber PivotFrame, FSequencerScriptingRange ScriptingRange, EMovieSceneTimeUnit TimeUnit = EMovieSceneTimeUnit::DisplayRate)
+	{
+		Impl::TransformKeysInChannel(ChannelHandle, OwningSequence, OwningSection, OffsetFrame, Scale, PivotFrame, ScriptingRange, TimeUnit);
+	}
+
 public:
 	TWeakObjectPtr<UMovieSceneSequence> OwningSequence;
 	TWeakObjectPtr<UMovieSceneSection> OwningSection;
diff --git a/Engine/Plugins/MovieScene/SequencerScripting/Source/SequencerScripting/Private/KeysAndChannels/MovieSceneScriptingChannel.h b/Engine/Plugins/MovieScene/SequencerScripting/Source/SequencerScripting/Private/KeysAndChannels/MovieSceneScriptingChannel.h
--- a/Engine/Plugins/MovieScene/SequencerScripting/Source/SequencerScripting/Private/KeysAndChannels/MovieSceneScriptingChannel.h
+++ b/Engine/Plugins/MovieScene/SequencerScripting/Source/SequencerScripting/Private/KeysAndChannels/MovieSceneScriptingChannel.h
@@ -7,6 +7,7 @@
 #include "SequencerScriptingRange.h"
 #include "MovieSceneTimeUnit.h"
 #include "ExtensionLibraries/MovieSceneSequenceExtensions.h"
+#include "IMovieSceneRetimingInterface.h"
 #include "MovieSceneSection.h"
 #include "MovieSceneSequence.h"
 #include "MovieScene.h"
@@ -313,6 +314,83 @@ struct TMovieSceneScriptingChannel
 		FFrame::KismetExecutionMessage(TEXT("Invalid ChannelHandle for MovieSceneScriptingChannel, failed to get default value."), ELogVerbosity::Error);
 		return TOptional<ScriptingKeyValueType>();
 	}
+
+	struct FTransformRetiming : UE::MovieScene::IRetimingInterface
+	{
+		FTransformRetiming(FFrameNumber InOffsetFrame, double InScale, FFrameNumber InPivotFrame, const TRange<FFrameNumber>& InRange)
+			: OffsetFrame(InOffsetFrame)
+			, Scale(InScale)
+			, PivotFrame(InPivotFrame)
+			, Range(InRange) {}
+
+		virtual FFrameTime RemapTime(FFrameTime InTime) const override
+		{
+			FFrameNumber FrameNumber = InTime.RoundToFrame();
+
+			if (Range.Contains(FrameNumber))
+			{
+				if (Scale != 0.0)
+				{
+					FrameNumber = (FrameNumber - PivotFrame) * Scale + PivotFrame;
+				}
+				return FrameNumber + OffsetFrame;
+			}
+			return InTime;
+		}
+		
+		double GetScale() const override { return 1.0; }
+		TUniquePtr<IRetimingInterface> RecurseInto(UMovieScene* InMovieScene) const override { return nullptr; }
+		void Begin(UMovieScene* InMovieScene) const override {}
+		void End(UMovieScene* InMovieScene) const override {}
+
+		FFrameNumber OffsetFrame;
+		double Scale = 0.0;
+		FFrameNumber PivotFrame;
+		TRange<FFrameNumber> Range;
+	};
+
+	static void TransformKeysInChannel(TMovieSceneChannelHandle<ChannelType> ChannelHandle, TWeakObjectPtr<UMovieSceneSequence> Sequence, TWeakObjectPtr<UMovieSceneSection> Section, FFrameNumber InOffsetFrame, double InScale, FFrameNumber InPivotFrame, FSequencerScriptingRange ScriptingRange, EMovieSceneTimeUnit TimeUnit)
+	{
+		// The Key's time is always going to be in Tick Resolution space, but the user may want to set it via Display Rate, so we convert.
+		FFrameNumber OffsetFrame = InOffsetFrame;
+		FFrameNumber PivotFrame = InPivotFrame;
+
+		UMovieSceneSequence* MovieSceneSequence = Sequence.Get();
+		FFrameRate Resolution = UMovieSceneSequenceExtensions::GetTickResolution(MovieSceneSequence);
+		TRange<FFrameNumber> Range = ScriptingRange.ToNative(Resolution);
+
+		if (Range.IsEmpty())
+		{
+			Range = TRange<FFrameNumber>::All();
+		}
+
+		if (TimeUnit == EMovieSceneTimeUnit::DisplayRate)
+		{
+			FFrameRate DisplayRate = UMovieSceneSequenceExtensions::GetDisplayRate(Sequence.Get());
+			FFrameRate TickResolution = UMovieSceneSequenceExtensions::GetTickResolution(Sequence.Get());
+
+			OffsetFrame = FFrameRate::TransformTime(OffsetFrame, DisplayRate, TickResolution).RoundToFrame();
+			PivotFrame = FFrameRate::TransformTime(PivotFrame, DisplayRate, TickResolution).RoundToFrame();
+
+			ScriptingRange.InternalRate = DisplayRate;
+			Range = ScriptingRange.ToNative(Resolution);
+		}
+
+		if (Section.IsValid())
+		{
+			Section->Modify();
+		}
+
+		ChannelType* Channel = ChannelHandle.Get();
+		if (Channel)
+		{
+			FTransformRetiming Retimer(OffsetFrame, InScale, PivotFrame, Range);
+			Channel->RemapTimes(Retimer);
+			return;
+		}
+
+		FFrame::KismetExecutionMessage(TEXT("Invalid ChannelHandle for MovieSceneScriptingChannel, failed to transform keys in channel."), ELogVerbosity::Error);
+	}
 };
 
 /** 
diff --git a/Engine/Plugins/MovieScene/SequencerScripting/Source/SequencerScripting/Private/KeysAndChannels/MovieSceneScriptingDouble.h b/Engine/Plugins/MovieScene/SequencerScripting/Source/SequencerScripting/Private/KeysAndChannels/MovieSceneScriptingDouble.h
--- a/Engine/Plugins/MovieScene/SequencerScripting/Source/SequencerScripting/Private/KeysAndChannels/MovieSceneScriptingDouble.h
+++ b/Engine/Plugins/MovieScene/SequencerScripting/Source/SequencerScripting/Private/KeysAndChannels/MovieSceneScriptingDouble.h
@@ -413,6 +413,22 @@ public:
 	{
 		return Impl::GetDefaultFromChannel(ChannelHandle).IsSet();
 	}
+
+	/**
+	 * Transform the keys in time in the channel by an offset, scale and pivot
+	 * @param	OffsetFrame	    The amount to offset the keys by
+	 * @param	Scale	        The amount to scale the keys by
+	 * @param	PivotFrame	    The frame to pivot around when scaling the keys
+	 * @param   ScriptingRange  The time range of the keys to transform
+	 * @param	TimeUnit 		Is the specified OffsetFrame, PivotFrame, and ScriptingRange in Display Rate frames or Tick Resolution.
+	 *
+	 */
+	UFUNCTION(BlueprintCallable, Category = "Sequencer|Keys", meta = (DisplayName = "Transform (Double)"))
+	void Transform(FFrameNumber OffsetFrame, double Scale, FFrameNumber PivotFrame, FSequencerScriptingRange ScriptingRange, EMovieSceneTimeUnit TimeUnit = EMovieSceneTimeUnit::DisplayRate)
+	{
+		Impl::TransformKeysInChannel(ChannelHandle, OwningSequence, OwningSection, OffsetFrame, Scale, PivotFrame, ScriptingRange, TimeUnit);
+	}
+
 public:
 	TWeakObjectPtr<UMovieSceneSequence> OwningSequence;
 	TMovieSceneChannelHandle<FMovieSceneDoubleChannel> ChannelHandle;
diff --git a/Engine/Plugins/MovieScene/SequencerScripting/Source/SequencerScripting/Private/KeysAndChannels/MovieSceneScriptingEvent.h b/Engine/Plugins/MovieScene/SequencerScripting/Source/SequencerScripting/Private/KeysAndChannels/MovieSceneScriptingEvent.h
--- a/Engine/Plugins/MovieScene/SequencerScripting/Source/SequencerScripting/Private/KeysAndChannels/MovieSceneScriptingEvent.h
+++ b/Engine/Plugins/MovieScene/SequencerScripting/Source/SequencerScripting/Private/KeysAndChannels/MovieSceneScriptingEvent.h
@@ -110,6 +110,21 @@ public:
 		return Impl::GetKeysInChannelByIndex(ChannelHandle, OwningSequence, OwningSection, Indices);
 	}
 
+	/**
+	 * Transform the keys in time in the channel by an offset, scale and pivot
+	 * @param	OffsetFrame	    The amount to offset the keys by
+	 * @param	Scale	        The amount to scale the keys by
+	 * @param	PivotFrame	    The frame to pivot around when scaling the keys
+	 * @param   ScriptingRange  The time range of the keys to transform
+	 * @param	TimeUnit 		Is the specified OffsetFrame, PivotFrame, and ScriptingRange in Display Rate frames or Tick Resolution.
+	 *
+	 */
+	UFUNCTION(BlueprintCallable, Category = "Sequencer|Keys", meta = (DisplayName = "Transform (Event)"))
+	void Transform(FFrameNumber OffsetFrame, double Scale, FFrameNumber PivotFrame, FSequencerScriptingRange ScriptingRange, EMovieSceneTimeUnit TimeUnit = EMovieSceneTimeUnit::DisplayRate)
+	{
+		Impl::TransformKeysInChannel(ChannelHandle, OwningSequence, OwningSection, OffsetFrame, Scale, PivotFrame, ScriptingRange, TimeUnit);
+	}
+
 public:
 	TWeakObjectPtr<UMovieSceneSequence> OwningSequence;
 	TMovieSceneChannelHandle<FMovieSceneEventChannel> ChannelHandle;
diff --git a/Engine/Plugins/MovieScene/SequencerScripting/Source/SequencerScripting/Private/KeysAndChannels/MovieSceneScriptingFloat.h b/Engine/Plugins/MovieScene/SequencerScripting/Source/SequencerScripting/Private/KeysAndChannels/MovieSceneScriptingFloat.h
--- a/Engine/Plugins/MovieScene/SequencerScripting/Source/SequencerScripting/Private/KeysAndChannels/MovieSceneScriptingFloat.h
+++ b/Engine/Plugins/MovieScene/SequencerScripting/Source/SequencerScripting/Private/KeysAndChannels/MovieSceneScriptingFloat.h
@@ -659,6 +659,29 @@ public:
 			return DoubleImpl::GetDefaultFromChannel(DoubleChannelHandle).IsSet();
 		}
 	}
+
+	/**
+	 * Transform the keys in time in the channel by an offset, scale and pivot
+	 * @param	OffsetFrame	    The amount to offset the keys by
+	 * @param	Scale	        The amount to scale the keys by
+	 * @param	PivotFrame	    The frame to pivot around when scaling the keys
+	 * @param   ScriptingRange  The time range of the keys to transform
+	 * @param	TimeUnit 		Is the specified OffsetFrame, PivotFrame, and ScriptingRange in Display Rate frames or Tick Resolution.
+	 *
+	 */
+	UFUNCTION(BlueprintCallable, Category = "Sequencer|Keys", meta = (DisplayName = "Transform (Float)"))
+	void Transform(FFrameNumber OffsetFrame, double Scale, FFrameNumber PivotFrame, FSequencerScriptingRange ScriptingRange, EMovieSceneTimeUnit TimeUnit = EMovieSceneTimeUnit::DisplayRate)
+	{
+		if (FloatChannelHandle.Get())
+		{
+			FloatImpl::TransformKeysInChannel(FloatChannelHandle, OwningSequence, OwningSection, OffsetFrame, Scale, PivotFrame, ScriptingRange, TimeUnit);
+		}
+		else
+		{
+			DoubleImpl::TransformKeysInChannel(DoubleChannelHandle, OwningSequence, OwningSection, OffsetFrame, Scale, PivotFrame, ScriptingRange, TimeUnit);
+		}
+	}
+
 public:
 	TWeakObjectPtr<UMovieSceneSequence> OwningSequence;
 	TWeakObjectPtr<UMovieSceneSection> OwningSection;
diff --git a/Engine/Plugins/MovieScene/SequencerScripting/Source/SequencerScripting/Private/KeysAndChannels/MovieSceneScriptingInteger.h b/Engine/Plugins/MovieScene/SequencerScripting/Source/SequencerScripting/Private/KeysAndChannels/MovieSceneScriptingInteger.h
--- a/Engine/Plugins/MovieScene/SequencerScripting/Source/SequencerScripting/Private/KeysAndChannels/MovieSceneScriptingInteger.h
+++ b/Engine/Plugins/MovieScene/SequencerScripting/Source/SequencerScripting/Private/KeysAndChannels/MovieSceneScriptingInteger.h
@@ -192,6 +192,21 @@ public:
 		return ChannelHandle.Get()->bInterpolateLinearKeys;
 	}
 
+	/**
+	 * Transform the keys in time in the channel by an offset, scale and pivot
+	 * @param	OffsetFrame	    The amount to offset the keys by
+	 * @param	Scale	        The amount to scale the keys by
+	 * @param	PivotFrame	    The frame to pivot around when scaling the keys
+	 * @param   ScriptingRange  The time range of the keys to transform
+	 * @param	TimeUnit 		Is the specified OffsetFrame, PivotFrame, and ScriptingRange in Display Rate frames or Tick Resolution.
+	 *
+	 */
+	UFUNCTION(BlueprintCallable, Category = "Sequencer|Keys", meta = (DisplayName = "Transform (Integer)"))
+	void Transform(FFrameNumber OffsetFrame, double Scale, FFrameNumber PivotFrame, FSequencerScriptingRange ScriptingRange, EMovieSceneTimeUnit TimeUnit = EMovieSceneTimeUnit::DisplayRate)
+	{
+		Impl::TransformKeysInChannel(ChannelHandle, OwningSequence, OwningSection, OffsetFrame, Scale, PivotFrame, ScriptingRange, TimeUnit);
+	}
+
 public:
 	TWeakObjectPtr<UMovieSceneSequence> OwningSequence;
 	TMovieSceneChannelHandle<FMovieSceneIntegerChannel> ChannelHandle;
diff --git a/Engine/Plugins/MovieScene/SequencerScripting/Source/SequencerScripting/Private/KeysAndChannels/MovieSceneScriptingObjectPath.h b/Engine/Plugins/MovieScene/SequencerScripting/Source/SequencerScripting/Private/KeysAndChannels/MovieSceneScriptingObjectPath.h
--- a/Engine/Plugins/MovieScene/SequencerScripting/Source/SequencerScripting/Private/KeysAndChannels/MovieSceneScriptingObjectPath.h
+++ b/Engine/Plugins/MovieScene/SequencerScripting/Source/SequencerScripting/Private/KeysAndChannels/MovieSceneScriptingObjectPath.h
@@ -201,6 +201,22 @@ public:
 	{
 		return GetDefault() != nullptr;
 	}
+
+	/**
+	 * Transform the keys in time in the channel by an offset, scale and pivot
+	 * @param	OffsetFrame	    The amount to offset the keys by
+	 * @param	Scale	        The amount to scale the keys by
+	 * @param	PivotFrame	    The frame to pivot around when scaling the keys
+	 * @param   ScriptingRange  The time range of the keys to transform
+	 * @param	TimeUnit 		Is the specified OffsetFrame, PivotFrame, and ScriptingRange in Display Rate frames or Tick Resolution.
+	 *
+	 */
+	UFUNCTION(BlueprintCallable, Category = "Sequencer|Keys", meta = (DisplayName = "Transform (Object Path)"))
+	void Transform(FFrameNumber OffsetFrame, double Scale, FFrameNumber PivotFrame, FSequencerScriptingRange ScriptingRange, EMovieSceneTimeUnit TimeUnit = EMovieSceneTimeUnit::DisplayRate)
+	{
+		Impl::TransformKeysInChannel(ChannelHandle, OwningSequence, OwningSection, OffsetFrame, Scale, PivotFrame, ScriptingRange, TimeUnit);
+	}
+
 public:
 	TWeakObjectPtr<UMovieSceneSequence> OwningSequence;
 	TWeakObjectPtr<UMovieSceneSection> OwningSection;
diff --git a/Engine/Plugins/MovieScene/SequencerScripting/Source/SequencerScripting/Private/KeysAndChannels/MovieSceneScriptingParticle.h b/Engine/Plugins/MovieScene/SequencerScripting/Source/SequencerScripting/Private/KeysAndChannels/MovieSceneScriptingParticle.h
--- a/Engine/Plugins/MovieScene/SequencerScripting/Source/SequencerScripting/Private/KeysAndChannels/MovieSceneScriptingParticle.h
+++ b/Engine/Plugins/MovieScene/SequencerScripting/Source/SequencerScripting/Private/KeysAndChannels/MovieSceneScriptingParticle.h
@@ -193,6 +193,22 @@ public:
 	{
 		return Impl::GetDefaultFromChannel(ChannelHandle).IsSet();
 	}
+
+	/**
+	 * Transform the keys in time in the channel by an offset, scale and pivot
+	 * @param	OffsetFrame	    The amount to offset the keys by
+	 * @param	Scale	        The amount to scale the keys by
+	 * @param	PivotFrame	    The frame to pivot around when scaling the keys
+	 * @param   ScriptingRange  The time range of the keys to transform
+	 * @param	TimeUnit 		Is the specified OffsetFrame, PivotFrame, and ScriptingRange in Display Rate frames or Tick Resolution.
+	 *
+	 */
+	UFUNCTION(BlueprintCallable, Category = "Sequencer|Keys", meta = (DisplayName = "Transform (Particle)"))
+	void Transform(FFrameNumber OffsetFrame, double Scale, FFrameNumber PivotFrame, FSequencerScriptingRange ScriptingRange, EMovieSceneTimeUnit TimeUnit = EMovieSceneTimeUnit::DisplayRate)
+	{
+		Impl::TransformKeysInChannel(ChannelHandle, OwningSequence, OwningSection, OffsetFrame, Scale, PivotFrame, ScriptingRange, TimeUnit);
+	}
+
 public:
 	TWeakObjectPtr<UMovieSceneSequence> OwningSequence;
 	TWeakObjectPtr<UMovieSceneSection> OwningSection;
diff --git a/Engine/Plugins/MovieScene/SequencerScripting/Source/SequencerScripting/Private/KeysAndChannels/MovieSceneScriptingString.h b/Engine/Plugins/MovieScene/SequencerScripting/Source/SequencerScripting/Private/KeysAndChannels/MovieSceneScriptingString.h
--- a/Engine/Plugins/MovieScene/SequencerScripting/Source/SequencerScripting/Private/KeysAndChannels/MovieSceneScriptingString.h
+++ b/Engine/Plugins/MovieScene/SequencerScripting/Source/SequencerScripting/Private/KeysAndChannels/MovieSceneScriptingString.h
@@ -153,6 +153,22 @@ public:
 	{
 		return Impl::GetDefaultFromChannel(ChannelHandle).IsSet();
 	}
+
+	/**
+	 * Transform the OffsetFrame in time in the channel by an offset, scale and pivot
+	 * @param	DeltaFrame	    The amount to offset the keys by
+	 * @param	Scale	        The amount to scale the keys by
+	 * @param	PivotFrame	    The frame to pivot around when scaling the keys
+	 * @param   ScriptingRange  The time range of the keys to transform
+	 * @param	TimeUnit 		Is the specified OffsetFrame, PivotFrame, and ScriptingRange in Display Rate frames or Tick Resolution.
+	 *
+	 */
+	UFUNCTION(BlueprintCallable, Category = "Sequencer|Keys", meta = (DisplayName = "Transform (String)"))
+	void Transform(FFrameNumber OffsetFrame, double Scale, FFrameNumber PivotFrame, FSequencerScriptingRange ScriptingRange, EMovieSceneTimeUnit TimeUnit = EMovieSceneTimeUnit::DisplayRate)
+	{
+		Impl::TransformKeysInChannel(ChannelHandle, OwningSequence, OwningSection, OffsetFrame, Scale, PivotFrame, ScriptingRange, TimeUnit);
+	}
+
 public:
 	TWeakObjectPtr<UMovieSceneSequence> OwningSequence;
 	TMovieSceneChannelHandle<FMovieSceneStringChannel> ChannelHandle;
