Sequencer Scripting: Follow up to 26776339 so that scripting an object property track as a class property can correctly set the bClassProperty value. This allows a class chooser instead of an object chooser to pop up in the details panel.

diff --git a/Engine/Plugins/MovieScene/SequencerScripting/Source/SequencerScripting/Private/ExtensionLibraries/MovieScenePropertyTrackExtensions.cpp b/Engine/Plugins/MovieScene/SequencerScripting/Source/SequencerScripting/Private/ExtensionLibraries/MovieScenePropertyTrackExtensions.cpp
--- a/Engine/Plugins/MovieScene/SequencerScripting/Source/SequencerScripting/Private/ExtensionLibraries/MovieScenePropertyTrackExtensions.cpp
+++ b/Engine/Plugins/MovieScene/SequencerScripting/Source/SequencerScripting/Private/ExtensionLibraries/MovieScenePropertyTrackExtensions.cpp
@@ -55,7 +55,7 @@ FName UMovieScenePropertyTrackExtensions::GetUniqueTrackName(UMovieSceneProperty
 #endif
 }
 
-void UMovieScenePropertyTrackExtensions::SetObjectPropertyClass(UMovieSceneObjectPropertyTrack* Track, UClass* PropertyClass)
+void UMovieScenePropertyTrackExtensions::SetObjectPropertyClass(UMovieSceneObjectPropertyTrack* Track, UClass* PropertyClass, bool bInClassProperty)
 {
 	if (!Track)
 	{
@@ -65,6 +65,9 @@ void UMovieScenePropertyTrackExtensions::SetObjectPropertyClass(UMovieSceneObjec
 
 	Track->Modify();
 	Track->PropertyClass = PropertyClass;
+#if WITH_EDITORONLY_DATA
+	Track->bClassProperty = bInClassProperty;
+#endif
 }
 
 UClass* UMovieScenePropertyTrackExtensions::GetObjectPropertyClass(UMovieSceneObjectPropertyTrack* Track)
diff --git a/Engine/Plugins/MovieScene/SequencerScripting/Source/SequencerScripting/Public/ExtensionLibraries/MovieScenePropertyTrackExtensions.h b/Engine/Plugins/MovieScene/SequencerScripting/Source/SequencerScripting/Public/ExtensionLibraries/MovieScenePropertyTrackExtensions.h
--- a/Engine/Plugins/MovieScene/SequencerScripting/Source/SequencerScripting/Public/ExtensionLibraries/MovieScenePropertyTrackExtensions.h
+++ b/Engine/Plugins/MovieScene/SequencerScripting/Source/SequencerScripting/Public/ExtensionLibraries/MovieScenePropertyTrackExtensions.h
@@ -62,9 +62,10 @@ public:
 	 *
 	 * @param Track        The track to use
 	 * @param PropertyClass The property class to set
+	 * @param bInClassProperty Whether the property class is a class property or not, which determines whether a class property chooser should be used instead of an object property chooser
 	 */
 	UFUNCTION(BlueprintCallable, Category = "Sequencer|Track", meta = (ScriptMethod))
-	static void SetObjectPropertyClass(UMovieSceneObjectPropertyTrack* Track, UClass* PropertyClass);
+	static void SetObjectPropertyClass(UMovieSceneObjectPropertyTrack* Track, UClass* PropertyClass, bool bInClassProperty = false);
 
 	/**
 	 * Get this object property track's property class
