Adding a new version of the ExportAnimSequence method to SequencerTools, which will spawn in all actors in the sequence and then wait for a provided delegate to return true before exporting. This can be used to get around issues caused by asynchronous loading of assets on mannequins and pawns.

#rnx

diff --git a/Engine/Plugins/MovieScene/SequencerScripting/Source/SequencerScriptingEditor/Private/SequencerTools.cpp b/Engine/Plugins/MovieScene/SequencerScripting/Source/SequencerScriptingEditor/Private/SequencerTools.cpp
--- a/Engine/Plugins/MovieScene/SequencerScripting/Source/SequencerScriptingEditor/Private/SequencerTools.cpp
+++ b/Engine/Plugins/MovieScene/SequencerScripting/Source/SequencerScriptingEditor/Private/SequencerTools.cpp
@@ -23,6 +23,7 @@
 #include "MovieScenePossessable.h"
 #include "ScopedTransaction.h"
 #include "Exporters/AnimSeqExportOption.h"
+#include "TimerManager.h"
 
 #include "K2Node_CustomEvent.h"
 #include "BlueprintFunctionNodeSpawner.h"
@@ -40,6 +41,47 @@
 
 #define LOCTEXT_NAMESPACE "SequencerTools"
 
+DECLARE_LOG_CATEGORY_CLASS(LogSequencer, Log, All);
+
+namespace SequencerToolsPrivate
+{
+	static TArray<TStrongObjectPtr<const UObject>> DelegateObjectsTempStrongRefs;
+
+	// Will check Condition every Interval seconds (or every frame, if Interval == 0) and execute Callback when true.
+	void WaitUntil(const FCheckCanCompleteDelegate Condition, const TFunction<void()> Callback, const UWorld* World, const float Interval = 0.1f)
+	{
+		if (!Condition.IsBound() || !Callback.IsSet() || !World)
+		{
+			UE_LOG(LogSequencer, Error, TEXT("WaitUntil() - One of the parameters is invalid."));
+			return;
+		}
+
+		// Guard the condition delegate from GC.
+		DelegateObjectsTempStrongRefs.AddUnique(TStrongObjectPtr<const UObject>(Condition.GetUObject()));
+
+		if (Condition.Execute())
+		{
+			Callback();
+			DelegateObjectsTempStrongRefs.Remove(TStrongObjectPtr<const UObject>(Condition.GetUObject()));
+			return;
+		}
+
+		FTimerDelegate WaitDelegate = FTimerDelegate::CreateLambda([Condition, Callback, World, Interval]() { SequencerToolsPrivate::WaitUntil(Condition, Callback, World, Interval); });
+		FTimerManager& TimerManager = World->GetTimerManager();
+
+		if (Interval == 0.f)
+		{
+			TimerManager.SetTimerForNextTick(WaitDelegate);
+		}
+		else
+		{
+			FTimerHandle Handle; // Unnecessary to keep, but required by SetTimer().
+			TimerManager.SetTimer(Handle, WaitDelegate, Interval, false);
+		}
+	}
+}
+
+
 bool USequencerToolsFunctionLibrary::RenderMovie(UMovieSceneCapture* InCaptureSettings, FOnRenderMovieStopped OnFinishedCallback)
 {
 	IMovieSceneCaptureDialogModule& MovieSceneCaptureModule = FModuleManager::Get().LoadModuleChecked<IMovieSceneCaptureDialogModule>("MovieSceneCaptureDialog");
@@ -65,7 +107,7 @@ PRAGMA_ENABLE_DEPRECATION_WARNINGS
 	{
 		if (!LevelSequenceCapture->LevelSequenceAsset.IsValid())
 		{
-			// UE_LOG(LogTemp, Warning, TEXT("No Level Sequence Asset specified in UAutomatedLevelSequenceCapture."));
+			// UE_LOG(LogSequencer, Warning, TEXT("No Level Sequence Asset specified in UAutomatedLevelSequenceCapture."));
 			FFrame::KismetExecutionMessage(TEXT("No Level Sequence Asset specified in UAutomatedLevelSequenceCapture."), ELogVerbosity::Error);
 			return false;
 		}
@@ -337,7 +379,13 @@ static USkeletalMeshComponent* GetSkelMeshComponent(IMovieScenePlayer* Player, c
 	return nullptr;
 }
 
-bool USequencerToolsFunctionLibrary::ExportAnimSequence(UWorld* World, ULevelSequence*  Sequence,  UAnimSequence* AnimSequence, UAnimSeqExportOption* ExportOptions,const FMovieSceneBindingProxy& Binding, bool bCreateLink)
+bool USequencerToolsFunctionLibrary::ExportAnimSequence(UWorld* World, ULevelSequence* Sequence, UAnimSequence* AnimSequence, UAnimSeqExportOption* ExportOptions, const FMovieSceneBindingProxy& Binding, bool bCreateLink)
+{
+	FCheckCanCompleteDelegate EmptyDelegate;
+	return ExportAnimSequenceWaitForDelegate(World, Sequence, AnimSequence, ExportOptions, Binding, bCreateLink, EmptyDelegate);
+}
+
+bool USequencerToolsFunctionLibrary::ExportAnimSequenceWaitForDelegate(UWorld* World, ULevelSequence* Sequence, UAnimSequence* AnimSequence, UAnimSeqExportOption* ExportOptions, const FMovieSceneBindingProxy& Binding, bool bCreateLink, const FCheckCanCompleteDelegate Delegate)
 {
 	UMovieScene* MovieScene = Sequence->GetMovieScene();
 	if (Binding.Sequence != Sequence || !AnimSequence)
@@ -353,47 +401,71 @@ bool USequencerToolsFunctionLibrary::ExportAnimSequence(UWorld* World, ULevelSeq
 	ULevelSequencePlayer* Player = ULevelSequencePlayer::CreateLevelSequencePlayer(World, Sequence, Settings, OutActor);
 
 	bool bResult = false;
-	
+
+	FSpawnableRestoreState SpawnableRestoreState(MovieScene, Player->GetSharedPlaybackState().ToSharedPtr());
+
+	if (SpawnableRestoreState.bWasChanged)
 	{
-		FSpawnableRestoreState SpawnableRestoreState(MovieScene, Player->GetSharedPlaybackState().ToSharedPtr());
- 
-		if (SpawnableRestoreState.bWasChanged)
+		// Evaluate at the beginning of the subscene time to ensure that spawnables are created before export
+		FFrameTime StartTime = FFrameRate::TransformTime(UE::MovieScene::DiscreteInclusiveLower(MovieScene->GetPlaybackRange()).Value, MovieScene->GetTickResolution(), MovieScene->GetDisplayRate());
+		Player->SetPlaybackPosition(FMovieSceneSequencePlaybackParams(StartTime, EUpdatePositionMethod::Play));
+
+		// Cannot transact changes due to in-flight spawnable transaction
+		ExportOptions->bTransactRecording = false;
+	}
+
+	auto CompleteExport = [Player, Binding, AnimSequence, RootToLocalTransform, Sequence, ExportOptions, World, OutActor, bCreateLink]() 
+	{
+		bool bResult = false;
+
+		if (Player && AnimSequence && Sequence && ExportOptions && World)
 		{
-			// Evaluate at the beginning of the subscene time to ensure that spawnables are created before export
-			FFrameTime StartTime = FFrameRate::TransformTime(UE::MovieScene::DiscreteInclusiveLower(MovieScene->GetPlaybackRange()).Value, MovieScene->GetTickResolution(), MovieScene->GetDisplayRate());
-			Player->SetPlaybackPosition(FMovieSceneSequencePlaybackParams(StartTime, EUpdatePositionMethod::Play));
+			USkeletalMeshComponent* SkeletalMeshComp = GetSkelMeshComponent(Player, Binding);
+			if (SkeletalMeshComp && SkeletalMeshComp->GetSkeletalMeshAsset() && SkeletalMeshComp->GetSkeletalMeshAsset()->GetSkeleton())
+			{
+				AnimSequence->SetSkeleton(SkeletalMeshComp->GetSkeletalMeshAsset()->GetSkeleton());
+				FAnimExportSequenceParameters AESQ;
+				AESQ.Player = Player;
+				AESQ.RootToLocalTransform = RootToLocalTransform;
+				AESQ.MovieSceneSequence = Sequence;
+				AESQ.RootMovieSceneSequence = Sequence;
+				bResult = MovieSceneToolHelpers::ExportToAnimSequence(AnimSequence, ExportOptions, AESQ, SkeletalMeshComp);
+			}
+			else
+			{
+				UE_LOG(LogSequencer, Warning, TEXT("USequencerToolsFunctionLibrary::ExportAnimSequence: No skel mesh found in Sequencer"));
+				bResult = false;
+			}
 
-			// Cannot transact changes due to in-flight spawnable transaction
-			ExportOptions->bTransactRecording = false;
+			Player->Stop();
+			World->DestroyActor(OutActor);
 		}
- 
-		USkeletalMeshComponent* SkeletalMeshComp =  GetSkelMeshComponent(Player, Binding);
-		if (SkeletalMeshComp && SkeletalMeshComp->GetSkeletalMeshAsset() && SkeletalMeshComp->GetSkeletalMeshAsset()->GetSkeleton())
+		else
 		{
-			AnimSequence->SetSkeleton(SkeletalMeshComp->GetSkeletalMeshAsset()->GetSkeleton());
-			FAnimExportSequenceParameters AESQ;
-			AESQ.Player = Player;
-			AESQ.RootToLocalTransform = RootToLocalTransform;
-			AESQ.MovieSceneSequence = Sequence;
-			AESQ.RootMovieSceneSequence = Sequence;
-			bResult = MovieSceneToolHelpers::ExportToAnimSequence(AnimSequence, ExportOptions, AESQ, SkeletalMeshComp);
+			UE_LOG(LogSequencer, Warning, TEXT("USequencerToolsFunctionLibrary::ExportAnimSequence: One or more of the completion lambda's captured parameters are invalid."));
+			bResult = false;
 		}
-		else
+
+		// create the link to the anim sequence
+		if (bResult && bCreateLink)
 		{
-			UE_LOG(LogTemp, Warning, TEXT("USequencerToolsFunctionLibrary::ExportAnimSequence: No skel mesh found in Sequencer"));
+			return LinkAnimSequence(Sequence, AnimSequence, ExportOptions, Binding);
 		}
-	}
-	
-	Player->Stop();
-	World->DestroyActor(OutActor);
 
-	//create the link to the anim sequence
-	if (bResult && bCreateLink)
+		return bResult;
+	};
+	
+	if (Delegate.IsBound())
 	{
-		return LinkAnimSequence(Sequence, AnimSequence, ExportOptions, Binding);
+		SequencerToolsPrivate::WaitUntil(Delegate, CompleteExport, World);
+		return true;
+	}
+	else
+	{
+		return CompleteExport();
 	}
-	return bResult;
 }
+
 void USequencerToolsFunctionLibrary::ClearLinkedAnimSequences(ULevelSequence* LevelSequence)
 {
 	if (LevelSequence)
diff --git a/Engine/Plugins/MovieScene/SequencerScripting/Source/SequencerScriptingEditor/Public/SequencerTools.h b/Engine/Plugins/MovieScene/SequencerScripting/Source/SequencerScriptingEditor/Public/SequencerTools.h
--- a/Engine/Plugins/MovieScene/SequencerScripting/Source/SequencerScriptingEditor/Public/SequencerTools.h
+++ b/Engine/Plugins/MovieScene/SequencerScripting/Source/SequencerScriptingEditor/Public/SequencerTools.h
@@ -32,6 +32,7 @@ class ULevelSequenceAnimSequenceLink;
 class UAnimSequenceLevelSequenceLink;
 
 DECLARE_DYNAMIC_DELEGATE_OneParam(FOnRenderMovieStopped, bool, bSuccess);
+DECLARE_DYNAMIC_DELEGATE_RetVal(bool, FCheckCanCompleteDelegate);
 
 
 USTRUCT(BlueprintType)
@@ -186,7 +187,13 @@ public:
 	 * @bCreateLink If true will create a link between the animation sequence and the level sequence
 	 */
 	UFUNCTION(BlueprintCallable, Category = "Editor Scripting | Sequencer Tools | Animation")
-	static UE_API bool ExportAnimSequence(UWorld* World, ULevelSequence*  Sequence, UAnimSequence* AnimSequence, UAnimSeqExportOption* ExportOption, const FMovieSceneBindingProxy& Binding, bool bCreateLink);
+	static UE_API bool ExportAnimSequence(UWorld* World, ULevelSequence*  Sequence, UAnimSequence* AnimSequence, UAnimSeqExportOption* ExportOptions, const FMovieSceneBindingProxy& Binding, bool bCreateLink);
+	
+	/*
+	* Version of ExportAnimSequence, which will spawn in all actors in the sequence and then wait for the provided Delegate to return true before exporting.
+	*/
+	UFUNCTION(BlueprintCallable, Category = "Editor Scripting | Sequencer Tools | Animation")
+	static UE_API bool ExportAnimSequenceWaitForDelegate(UWorld* World, ULevelSequence* Sequence, UAnimSequence* AnimSequence, UAnimSeqExportOption* ExportOptions, const FMovieSceneBindingProxy& Binding, bool bCreateLink, const FCheckCanCompleteDelegate Delegate);
 
 	/*
 	 * Links a LevelSequence's SkeletalMesh binding to an existing anim sequence.
