Sequencer: MovieSceneSequenceExtensions::FindBindingByName now returns a more reliable binding because it gathers the names from FMovieScenePossessable and FMovieSceneSpawnable which are kept in sync with the actor labels.

diff --git a/Engine/Plugins/MovieScene/SequencerScripting/Source/SequencerScripting/Private/ExtensionLibraries/MovieSceneSequenceExtensions.cpp b/Engine/Plugins/MovieScene/SequencerScripting/Source/SequencerScripting/Private/ExtensionLibraries/MovieSceneSequenceExtensions.cpp
--- a/Engine/Plugins/MovieScene/SequencerScripting/Source/SequencerScripting/Private/ExtensionLibraries/MovieSceneSequenceExtensions.cpp
+++ b/Engine/Plugins/MovieScene/SequencerScripting/Source/SequencerScripting/Private/ExtensionLibraries/MovieSceneSequenceExtensions.cpp
@@ -676,13 +676,25 @@ FMovieSceneBindingProxy UMovieSceneSequenceExtensions::FindBindingByName(UMovieS
 		return FMovieSceneBindingProxy();
 	}
 
-	const UMovieScene* MovieScene = GetMovieScene(Sequence);
+	UMovieScene* MovieScene = GetMovieScene(Sequence);
 	if (MovieScene)
 	{
-		const FMovieSceneBinding* Binding = Algo::FindBy(MovieScene->GetBindings(), Name, &FMovieSceneBinding::GetName);
-		if (Binding)
+		for (int32 Index = 0; MovieScene->GetPossessableCount(); ++Index)
 		{
-			return FMovieSceneBindingProxy(Binding->GetObjectGuid(), Sequence);
+			const FMovieScenePossessable& Possessable = MovieScene->GetPossessable(Index);
+			if (Possessable.GetName() == Name)
+			{
+				return FMovieSceneBindingProxy(Possessable.GetGuid(), Sequence);
+			}
+		}
+
+		for (int32 Index = 0; MovieScene->GetSpawnableCount(); ++Index)
+		{
+			const FMovieSceneSpawnable& Spawnable = MovieScene->GetSpawnable(Index);
+			if (Spawnable.GetName() == Name)
+			{
+				return FMovieSceneBindingProxy(Spawnable.GetGuid(), Sequence);
+			}
 		}
 	}
 	return FMovieSceneBindingProxy();
