Sequencer: MovieSceneSequenceExtensions::FindBindingByName now returns a more reliable binding because it gathers the names from FMovieScenePossessable and FMovieSceneSpawnable which are kept in sync with the actor labels.

diff --git a/Engine/Plugins/MovieScene/SequencerScripting/Source/SequencerScripting/Private/ExtensionLibraries/MovieSceneSequenceExtensions.cpp b/Engine/Plugins/MovieScene/SequencerScripting/Source/SequencerScripting/Private/ExtensionLibraries/MovieSceneSequenceExtensions.cpp
--- a/Engine/Plugins/MovieScene/SequencerScripting/Source/SequencerScripting/Private/ExtensionLibraries/MovieSceneSequenceExtensions.cpp
+++ b/Engine/Plugins/MovieScene/SequencerScripting/Source/SequencerScripting/Private/ExtensionLibraries/MovieSceneSequenceExtensions.cpp
@@ -679,7 +679,7 @@ FMovieSceneBindingProxy UMovieSceneSequenceExtensions::FindBindingByName(UMovieS
 	UMovieScene* MovieScene = GetMovieScene(Sequence);
 	if (MovieScene)
 	{
-		for (int32 Index = 0; MovieScene->GetPossessableCount(); ++Index)
+		for (int32 Index = 0; Index < MovieScene->GetPossessableCount(); ++Index)
 		{
 			const FMovieScenePossessable& Possessable = MovieScene->GetPossessable(Index);
 			if (Possessable.GetName() == Name)
@@ -688,7 +688,7 @@ FMovieSceneBindingProxy UMovieSceneSequenceExtensions::FindBindingByName(UMovieS
 			}
 		}
 
-		for (int32 Index = 0; MovieScene->GetSpawnableCount(); ++Index)
+		for (int32 Index = 0; Index < MovieScene->GetSpawnableCount(); ++Index)
 		{
 			const FMovieSceneSpawnable& Spawnable = MovieScene->GetSpawnable(Index);
 			if (Spawnable.GetName() == Name)
