Sequencer: UE_API all blueprint exposed functions

#FailedValidations Preflight

diff --git a/Engine/Plugins/MovieScene/SequencerScripting/Source/SequencerScripting/Public/ExtensionLibraries/MovieSceneBindingExtensions.h b/Engine/Plugins/MovieScene/SequencerScripting/Source/SequencerScripting/Public/ExtensionLibraries/MovieSceneBindingExtensions.h
--- a/Engine/Plugins/MovieScene/SequencerScripting/Source/SequencerScripting/Public/ExtensionLibraries/MovieSceneBindingExtensions.h
+++ b/Engine/Plugins/MovieScene/SequencerScripting/Source/SequencerScripting/Public/ExtensionLibraries/MovieSceneBindingExtensions.h
@@ -10,6 +10,7 @@ class UMovieSceneTrack;
 struct FMovieSceneBindingProxy;
 template <typename T> class TSubclassOf;
 
+#define UE_API SEQUENCERSCRIPTING_API
 
 /**
  * Function library containing methods that should be hoisted onto FMovieSceneBindingProxies for scripting
@@ -23,15 +24,15 @@ public:
 	 * Check whether the specified binding is valid
 	 */
 	UFUNCTION(BlueprintPure, Category = "Sequencer|Sequence", meta=(ScriptMethod))
-	static bool IsValid(const FMovieSceneBindingProxy& InBinding);
+	static UE_API bool IsValid(const FMovieSceneBindingProxy& InBinding);
 
 	/** Returns true if the values are equal (A == B) */
 	UFUNCTION(BlueprintPure, meta = (DisplayName = "Equal (MovieSceneBindingProxy)", CompactNodeTitle = "==", ScriptOperator = "==", BlueprintThreadSafe), Category = "Sequencer|Sequence")
-	static bool EqualEqual_MovieSceneBindingProxy(FMovieSceneBindingProxy A, FMovieSceneBindingProxy B);
+	static UE_API bool EqualEqual_MovieSceneBindingProxy(FMovieSceneBindingProxy A, FMovieSceneBindingProxy B);
 
 	/** Returns true if the values are not equal (A != B) */
 	UFUNCTION(BlueprintPure, meta = (DisplayName = "Not Equal (MovieSceneBindingProxy)", CompactNodeTitle = "!=", ScriptOperator = "!=", BlueprintThreadSafe), Category = "Sequencer|Sequence")
-	static bool NotEqual_MovieSceneBindingProxy(FMovieSceneBindingProxy A, FMovieSceneBindingProxy B);
+	static UE_API bool NotEqual_MovieSceneBindingProxy(FMovieSceneBindingProxy A, FMovieSceneBindingProxy B);
 
 	/**
 	 * Get this binding's ID
@@ -40,7 +41,7 @@ public:
 	 * @return The guid that uniquely represents this binding
 	 */
 	UFUNCTION(BlueprintCallable, Category = "Sequencer|Sequence", meta=(ScriptMethod))
-	static FGuid GetId(const FMovieSceneBindingProxy& InBinding);
+	static UE_API FGuid GetId(const FMovieSceneBindingProxy& InBinding);
 
 	/**
 	 * Get this binding's name
@@ -49,7 +50,7 @@ public:
 	 * @return The display name of the binding
 	 */
 	UFUNCTION(BlueprintCallable, Category = "Sequencer|Sequence", meta=(ScriptMethod))
-	static FText GetDisplayName(const FMovieSceneBindingProxy& InBinding);
+	static UE_API FText GetDisplayName(const FMovieSceneBindingProxy& InBinding);
 
 	/**
 	 * Set this binding's name
@@ -58,7 +59,7 @@ public:
 	 * @param InName The display name of the binding
 	 */
 	UFUNCTION(BlueprintCallable, Category = "Sequencer|Sequence", meta=(ScriptMethod, DevelopmentOnly))
-	static void SetDisplayName(const FMovieSceneBindingProxy& InBinding, const FText& InDisplayName);
+	static UE_API void SetDisplayName(const FMovieSceneBindingProxy& InBinding, const FText& InDisplayName);
 
 	/**
 	 * Get this binding's object non-display name
@@ -67,7 +68,7 @@ public:
 	 * @return The name of the binding
 	 */
 	UFUNCTION(BlueprintCallable, Category = "Sequencer|Sequence", meta=(ScriptMethod))
-	static FString GetName(const FMovieSceneBindingProxy& InBinding);
+	static UE_API FString GetName(const FMovieSceneBindingProxy& InBinding);
 
 	/**
 	 * Set this binding's object non-display name
@@ -76,7 +77,7 @@ public:
 	 * @param InName The name of the binding
 	 */
 	UFUNCTION(BlueprintCallable, Category = "Sequencer|Sequence", meta=(ScriptMethod))
-	static void SetName(const FMovieSceneBindingProxy& InBinding, const FString& InName);
+	static UE_API void SetName(const FMovieSceneBindingProxy& InBinding, const FString& InName);
 
 	/**
 	 * Get the sorting order for this binding
@@ -85,7 +86,7 @@ public:
 	 * @return The sorting order of the requested binding
 	 */
 	UFUNCTION(BlueprintCallable, Category = "Sequencer|Sequence", meta = (ScriptMethod, DevelopmentOnly))
-	static int32 GetSortingOrder(const FMovieSceneBindingProxy& InBinding);
+	static UE_API int32 GetSortingOrder(const FMovieSceneBindingProxy& InBinding);
 
 	/**
 	 * Set the sorting order for this binding
@@ -94,7 +95,7 @@ public:
 	 * @param SortingOrder The sorting order to set
 	 */
 	UFUNCTION(BlueprintCallable, Category = "Sequencer|Sequence", meta = (ScriptMethod, DevelopmentOnly))
-	static void SetSortingOrder(const FMovieSceneBindingProxy& InBinding, int32 SortingOrder);
+	static UE_API void SetSortingOrder(const FMovieSceneBindingProxy& InBinding, int32 SortingOrder);
 
 	/**
 	 * Get all the tracks stored within this binding
@@ -103,7 +104,7 @@ public:
 	 * @return An array containing all the binding's tracks
 	 */
 	UFUNCTION(BlueprintCallable, Category = "Sequencer|Sequence", meta=(ScriptMethod))
-	static TArray<UMovieSceneTrack*> GetTracks(const FMovieSceneBindingProxy& InBinding);
+	static UE_API TArray<UMovieSceneTrack*> GetTracks(const FMovieSceneBindingProxy& InBinding);
 
 	/**
 	 * Find all tracks within a given binding of the specified type
@@ -113,7 +114,7 @@ public:
 	 * @return An array containing any tracks that match the type specified
 	 */
 	UFUNCTION(BlueprintCallable, Category = "Sequencer|Sequence", meta=(ScriptMethod, DeterminesOutputType="TrackType"))
-	static TArray<UMovieSceneTrack*> FindTracksByType(const FMovieSceneBindingProxy& InBinding, TSubclassOf<UMovieSceneTrack> TrackType);
+	static UE_API TArray<UMovieSceneTrack*> FindTracksByType(const FMovieSceneBindingProxy& InBinding, TSubclassOf<UMovieSceneTrack> TrackType);
 
 	/**
 	 * Find all tracks within a given binding of the specified type, not allowing sub-classed types
@@ -123,7 +124,7 @@ public:
 	 * @return An array containing any tracks that are exactly the same as the type specified
 	 */
 	UFUNCTION(BlueprintCallable, Category = "Sequencer|Sequence", meta=(ScriptMethod, DeterminesOutputType="TrackType"))
-	static TArray<UMovieSceneTrack*> FindTracksByExactType(const FMovieSceneBindingProxy& InBinding, TSubclassOf<UMovieSceneTrack> TrackType);
+	static UE_API TArray<UMovieSceneTrack*> FindTracksByExactType(const FMovieSceneBindingProxy& InBinding, TSubclassOf<UMovieSceneTrack> TrackType);
 
 	/**
 	 * Remove the specified track from this binding
@@ -132,7 +133,7 @@ public:
 	 * @param TrackToRemove The track to remove
 	 */
 	UFUNCTION(BlueprintCallable, Category = "Sequencer|Sequence", meta=(ScriptMethod))
-	static void RemoveTrack(const FMovieSceneBindingProxy& InBinding, UMovieSceneTrack* TrackToRemove);
+	static UE_API void RemoveTrack(const FMovieSceneBindingProxy& InBinding, UMovieSceneTrack* TrackToRemove);
 
 	/**
 	 * Remove the specified binding
@@ -140,7 +141,7 @@ public:
 	 * @param InBinding     The binding to remove the track from
 	 */
 	UFUNCTION(BlueprintCallable, Category = "Sequencer|Sequence", meta=(ScriptMethod))
-	static void Remove(const FMovieSceneBindingProxy& InBinding);
+	static UE_API void Remove(const FMovieSceneBindingProxy& InBinding);
 
 	/**
 	 * Add a new track to the specified binding
@@ -150,7 +151,7 @@ public:
 	 * @return The newly created track, if successful
 	 */
 	UFUNCTION(BlueprintCallable, Category = "Sequencer|Sequence", meta=(ScriptMethod))
-	static UMovieSceneTrack* AddTrack(const FMovieSceneBindingProxy& InBinding, TSubclassOf<UMovieSceneTrack> TrackType);
+	static UE_API UMovieSceneTrack* AddTrack(const FMovieSceneBindingProxy& InBinding, TSubclassOf<UMovieSceneTrack> TrackType);
 
 	/**
 	* Get all the children of this binding
@@ -159,7 +160,7 @@ public:
 	* @return An array containing all the binding's children
 	*/
 	UFUNCTION(BlueprintCallable, Category = "Sequencer|Sequence", meta=(ScriptMethod))
-	static TArray<FMovieSceneBindingProxy> GetChildPossessables(const FMovieSceneBindingProxy& InBinding);
+	static UE_API TArray<FMovieSceneBindingProxy> GetChildPossessables(const FMovieSceneBindingProxy& InBinding);
 
 	/**
 	* Get this binding's object template
@@ -168,7 +169,7 @@ public:
 	* @return The object template of the binding
 	*/
 	UFUNCTION(BlueprintCallable, Category = "Sequencer|Sequence", meta=(ScriptMethod))
-	static UObject* GetObjectTemplate(const FMovieSceneBindingProxy& InBinding);
+	static UE_API UObject* GetObjectTemplate(const FMovieSceneBindingProxy& InBinding);
 
 	/**
 	* Get this binding's possessed object class
@@ -177,7 +178,7 @@ public:
 	* @return The possessed object class of the binding
 	*/
 	UFUNCTION(BlueprintCallable, Category = "Sequencer|Sequence", meta=(ScriptMethod))
-	static UClass* GetPossessedObjectClass(const FMovieSceneBindingProxy& InBinding);
+	static UE_API UClass* GetPossessedObjectClass(const FMovieSceneBindingProxy& InBinding);
 
 	/**
 	* Get the parent of this binding
@@ -186,7 +187,7 @@ public:
 	* @return The binding's parent
 	*/
 	UFUNCTION(BlueprintCallable, Category = "Sequencer|Sequence", meta=(ScriptMethod))
-	static FMovieSceneBindingProxy GetParent(const FMovieSceneBindingProxy& InBinding);
+	static UE_API FMovieSceneBindingProxy GetParent(const FMovieSceneBindingProxy& InBinding);
 
 	/**
 	 * Set the parent to this binding
@@ -195,7 +196,7 @@ public:
 	 * @param InParentBinding     The parent to set the InBinding to
 	 */
 	UFUNCTION(BlueprintCallable, Category = "Sequencer|Sequence", meta = (ScriptMethod))
-	static void SetParent(const FMovieSceneBindingProxy& InBinding, const FMovieSceneBindingProxy& InParentBinding);
+	static UE_API void SetParent(const FMovieSceneBindingProxy& InBinding, const FMovieSceneBindingProxy& InParentBinding);
 
 	/**
      * Move all the contents (tracks, child bindings) of the specified binding ID onto another
@@ -204,7 +205,7 @@ public:
 	 * @param DestinationBindingId The identifier of the binding ID to move the contents to	 
 	 */
 	UFUNCTION(BlueprintCallable, Category = "Sequencer|Sequence", meta = (ScriptMethod))
-	static void MoveBindingContents(const FMovieSceneBindingProxy& SourceBindingId, const FMovieSceneBindingProxy& DestinationBindingId);
+	static UE_API void MoveBindingContents(const FMovieSceneBindingProxy& SourceBindingId, const FMovieSceneBindingProxy& DestinationBindingId);
 
 	/**
 	 * Set the spawnable id that the possessable binding should possess
@@ -213,5 +214,7 @@ public:
 	 * @param SpawnableBindingID The spawnable's binding id
 	 */
 	UFUNCTION(BlueprintCallable, Category = "Sequencer|Sequence", meta = (ScriptMethod))
-	static void SetSpawnableBindingID(const FMovieSceneBindingProxy& InBinding, const FMovieSceneObjectBindingID& SpawnableBindingID);
+	static UE_API void SetSpawnableBindingID(const FMovieSceneBindingProxy& InBinding, const FMovieSceneObjectBindingID& SpawnableBindingID);
 };
+
+#undef UE_API
diff --git a/Engine/Plugins/MovieScene/SequencerScripting/Source/SequencerScripting/Public/ExtensionLibraries/MovieSceneEventTrackExtensions.h b/Engine/Plugins/MovieScene/SequencerScripting/Source/SequencerScripting/Public/ExtensionLibraries/MovieSceneEventTrackExtensions.h
--- a/Engine/Plugins/MovieScene/SequencerScripting/Source/SequencerScripting/Public/ExtensionLibraries/MovieSceneEventTrackExtensions.h
+++ b/Engine/Plugins/MovieScene/SequencerScripting/Source/SequencerScripting/Public/ExtensionLibraries/MovieSceneEventTrackExtensions.h
@@ -11,6 +11,8 @@ class UMovieSceneEventTrack;
 class UMovieSceneEventRepeaterSection;
 class UMovieSceneEventTriggerSection;
 
+#define UE_API SEQUENCERSCRIPTING_API
+
 /**
  * Function library containing methods that should be hoisted onto UMovieSceneEventTrack for scripting
  */
@@ -28,7 +30,7 @@ public:
 	 * @return The newly created event repeater section
 	 */
 	UFUNCTION(BlueprintCallable, Category = "Sequencer|Track", meta = (ScriptMethod))
-	static UMovieSceneEventRepeaterSection* AddEventRepeaterSection(UMovieSceneEventTrack* InTrack);
+	static UE_API UMovieSceneEventRepeaterSection* AddEventRepeaterSection(UMovieSceneEventTrack* InTrack);
 	
 	/**
 	 * Create a new event trigger section for the given track
@@ -37,7 +39,7 @@ public:
 	 * @return The newly created event trigger section
 	 */
 	UFUNCTION(BlueprintCallable, Category = "Sequencer|Track", meta = (ScriptMethod))
-	static UMovieSceneEventTriggerSection* AddEventTriggerSection(UMovieSceneEventTrack* InTrack);
+	static UE_API UMovieSceneEventTriggerSection* AddEventTriggerSection(UMovieSceneEventTrack* InTrack);
 	
 	/*
 	 * Return the class of the bound object property
@@ -48,3 +50,5 @@ public:
 	UFUNCTION(BlueprintCallable, Category="Sequencer|Event", meta = (ScriptMethod))
 	static UClass* GetBoundObjectPropertyClass(const FMovieSceneEvent& EventKey) { return EventKey.GetBoundObjectPropertyClass(); }
 };
+
+#undef UE_API
diff --git a/Engine/Plugins/MovieScene/SequencerScripting/Source/SequencerScripting/Public/ExtensionLibraries/MovieSceneFolderExtensions.h b/Engine/Plugins/MovieScene/SequencerScripting/Source/SequencerScripting/Public/ExtensionLibraries/MovieSceneFolderExtensions.h
--- a/Engine/Plugins/MovieScene/SequencerScripting/Source/SequencerScripting/Public/ExtensionLibraries/MovieSceneFolderExtensions.h
+++ b/Engine/Plugins/MovieScene/SequencerScripting/Source/SequencerScripting/Public/ExtensionLibraries/MovieSceneFolderExtensions.h
@@ -11,6 +11,8 @@ struct FMovieSceneBindingProxy;
 class UMovieSceneFolder;
 class UMovieSceneTrack;
 
+#define UE_API SEQUENCERSCRIPTING_API
+
 /**
  * Function library containing methods that should be hoisted onto UMovieSceneFolders for scripting
  */
@@ -27,7 +29,7 @@ public:
 	 * @return The target folder's name
 	 */
 	UFUNCTION(BlueprintCallable, BlueprintPure, Category = "Sequencer|Sequence", meta=(ScriptMethod))
-	static FName GetFolderName(UMovieSceneFolder* Folder);
+	static UE_API FName GetFolderName(UMovieSceneFolder* Folder);
 
 	/**
 	 * Set the name of the given folder
@@ -37,7 +39,7 @@ public:
 	 * @return True if the setting of the folder name succeeds
 	 */
 	UFUNCTION(BlueprintCallable, Category = "Sequencer|Sequence", meta=(ScriptMethod))
-	static bool SetFolderName(UMovieSceneFolder* Folder, FName InFolderName);
+	static UE_API bool SetFolderName(UMovieSceneFolder* Folder, FName InFolderName);
 
 	/**
 	 * Get the display color of the given folder
@@ -46,7 +48,7 @@ public:
 	 * @return The display color of the given folder
 	 */
 	UFUNCTION(BlueprintCallable, BlueprintPure, Category = "Sequencer|Sequence", meta = (ScriptMethod, DevelopmentOnly))
-	static FColor GetFolderColor(UMovieSceneFolder* Folder);
+	static UE_API FColor GetFolderColor(UMovieSceneFolder* Folder);
 
 	/**
 	 * Set the display color of the given folder
@@ -56,7 +58,7 @@ public:
 	 * @return True if the folder's display color is set successfully
 	 */
 	UFUNCTION(BlueprintCallable, Category = "Sequencer|Sequence", meta=(ScriptMethod, DevelopmentOnly))
-	static bool SetFolderColor(UMovieSceneFolder* Folder, FColor InFolderColor);
+	static UE_API bool SetFolderColor(UMovieSceneFolder* Folder, FColor InFolderColor);
 
 	/**
 	 * Get the child folders of a given folder
@@ -65,7 +67,7 @@ public:
 	 * @return The child folders associated with the given folder
 	 */
 	UFUNCTION(BlueprintCallable, BlueprintPure, Category = "Sequencer|Sequence", meta=(ScriptMethod))
-	static TArray<UMovieSceneFolder*> GetChildFolders(UMovieSceneFolder* Folder);
+	static UE_API TArray<UMovieSceneFolder*> GetChildFolders(UMovieSceneFolder* Folder);
 
 	/**
 	 * Add a child folder to the target folder
@@ -75,7 +77,7 @@ public:
 	 * @return True if the addition is successful
 	 */
 	UFUNCTION(BlueprintCallable, Category = "Sequencer|Sequence", meta=(ScriptMethod))
-	static bool AddChildFolder(UMovieSceneFolder* TargetFolder, UMovieSceneFolder* FolderToAdd);
+	static UE_API bool AddChildFolder(UMovieSceneFolder* TargetFolder, UMovieSceneFolder* FolderToAdd);
 
 	/**
 	 * Remove a child folder from the given folder
@@ -85,7 +87,7 @@ public:
 	 * @return True if the removal succeeds
 	 */
 	UFUNCTION(BlueprintCallable, Category = "Sequencer|Sequence", meta=(ScriptMethod))
-	static bool RemoveChildFolder(UMovieSceneFolder* TargetFolder, UMovieSceneFolder* FolderToRemove);
+	static UE_API bool RemoveChildFolder(UMovieSceneFolder* TargetFolder, UMovieSceneFolder* FolderToRemove);
 
 	/**
 	 * Get the tracks contained by this folder
@@ -94,7 +96,7 @@ public:
 	 * @return The tracks under the given folder
 	 */
 	UFUNCTION(BlueprintCallable, BlueprintPure, Category = "Sequencer|Sequence", meta = (ScriptMethod))
-	static TArray<UMovieSceneTrack*> GetChildTracks(UMovieSceneFolder* Folder);
+	static UE_API TArray<UMovieSceneTrack*> GetChildTracks(UMovieSceneFolder* Folder);
 	
 	/**
 	 * Add a track to this folder
@@ -104,7 +106,7 @@ public:
 	 * @return True if the addition is successful
 	 */
 	UFUNCTION(BlueprintCallable, Category = "Sequencer|Sequence", meta = (ScriptMethod))
-	static bool AddChildTrack(UMovieSceneFolder* Folder, UMovieSceneTrack* InTrack);
+	static UE_API bool AddChildTrack(UMovieSceneFolder* Folder, UMovieSceneTrack* InTrack);
 	
 	/**
 	 * Remove a track from the given folder
@@ -114,7 +116,7 @@ public:
 	 * @return True if the removal succeeds
 	 */
 	UFUNCTION(BlueprintCallable, Category = "Sequencer|Sequence", meta=(ScriptMethod))
-	static bool RemoveChildTrack(UMovieSceneFolder* Folder, UMovieSceneTrack* InTrack);
+	static UE_API bool RemoveChildTrack(UMovieSceneFolder* Folder, UMovieSceneTrack* InTrack);
 
 	/**
 	 * Get the object bindings contained by this folder
@@ -123,7 +125,7 @@ public:
 	 * @return The object bindings under the given folder
 	 */
 	UFUNCTION(BlueprintCallable, BlueprintPure, Category = "Sequencer|Sequence", meta=(ScriptMethod))
-	static TArray<FMovieSceneBindingProxy> GetChildObjectBindings(UMovieSceneFolder* Folder);
+	static UE_API TArray<FMovieSceneBindingProxy> GetChildObjectBindings(UMovieSceneFolder* Folder);
 
 	/** 
 	 * Add a guid for an object binding to this folder 
@@ -133,7 +135,7 @@ public:
 	 * @return True if the addition is successful
 	 */
 	UFUNCTION(BlueprintCallable, Category = "Sequencer|Sequence", meta=(ScriptMethod))
-	static bool AddChildObjectBinding(UMovieSceneFolder* Folder, FMovieSceneBindingProxy InObjectBinding);
+	static UE_API bool AddChildObjectBinding(UMovieSceneFolder* Folder, FMovieSceneBindingProxy InObjectBinding);
 
 	/** 
 	 * Remove an object binding from the given folder
@@ -143,5 +145,7 @@ public:
 	 * @return True if the operation succeeds
 	 */
 	UFUNCTION(BlueprintCallable, Category = "Sequencer|Sequence", meta=(ScriptMethod))
-	static bool RemoveChildObjectBinding(UMovieSceneFolder* Folder, const FMovieSceneBindingProxy InObjectBinding);
+	static UE_API bool RemoveChildObjectBinding(UMovieSceneFolder* Folder, const FMovieSceneBindingProxy InObjectBinding);
 };
+
+#undef UE_API
diff --git a/Engine/Plugins/MovieScene/SequencerScripting/Source/SequencerScripting/Public/ExtensionLibraries/MovieSceneMaterialTrackExtensions.h b/Engine/Plugins/MovieScene/SequencerScripting/Source/SequencerScripting/Public/ExtensionLibraries/MovieSceneMaterialTrackExtensions.h
--- a/Engine/Plugins/MovieScene/SequencerScripting/Source/SequencerScripting/Public/ExtensionLibraries/MovieSceneMaterialTrackExtensions.h
+++ b/Engine/Plugins/MovieScene/SequencerScripting/Source/SequencerScripting/Public/ExtensionLibraries/MovieSceneMaterialTrackExtensions.h
@@ -12,6 +12,8 @@ class UMovieSceneMaterialTrack;
 
 struct FComponentMaterialInfo;
 
+#define UE_API SEQUENCERSCRIPTING_API
+
 /**
  * Function library containing methods that should be hoisted onto UMovieSceneMaterialTrack for scripting
  */
@@ -46,7 +48,7 @@ public:
 	 * @param MaterialInfo The desired material info to animate.
 	 */
 	UFUNCTION(BlueprintCallable, Category = "Sequencer|Track", meta = (ScriptMethod))
-	static void SetMaterialInfo(UMovieSceneComponentMaterialTrack* Track, const FComponentMaterialInfo& MaterialInfo);
+	static UE_API void SetMaterialInfo(UMovieSceneComponentMaterialTrack* Track, const FComponentMaterialInfo& MaterialInfo);
 
 	/**
 	 * Get material info of the component that is animated by the material track.
@@ -54,5 +56,7 @@ public:
 	 * @return The material info.
 	 */
 	UFUNCTION(BlueprintCallable, Category = "Sequencer|Track", meta=(ScriptMethod, DeprecatedFunction, DeprecationMessage = "Use SetMaterialInfo instead"))
-	static FComponentMaterialInfo GetMaterialInfo(UMovieSceneComponentMaterialTrack* Track);
+	static UE_API FComponentMaterialInfo GetMaterialInfo(UMovieSceneComponentMaterialTrack* Track);
 };
+
+#undef UE_API
diff --git a/Engine/Plugins/MovieScene/SequencerScripting/Source/SequencerScripting/Public/ExtensionLibraries/MovieScenePrimitiveMaterialTrackExtensions.h b/Engine/Plugins/MovieScene/SequencerScripting/Source/SequencerScripting/Public/ExtensionLibraries/MovieScenePrimitiveMaterialTrackExtensions.h
--- a/Engine/Plugins/MovieScene/SequencerScripting/Source/SequencerScripting/Public/ExtensionLibraries/MovieScenePrimitiveMaterialTrackExtensions.h
+++ b/Engine/Plugins/MovieScene/SequencerScripting/Source/SequencerScripting/Public/ExtensionLibraries/MovieScenePrimitiveMaterialTrackExtensions.h
@@ -8,6 +8,8 @@
 
 class UMovieScenePrimitiveMaterialTrack;
 
+#define UE_API SEQUENCERSCRIPTING_API
+
 /**
  * Function library containing methods that should be hoisted onto UMovieScenePrimitiveMaterialTrack for scripting
  */
@@ -42,7 +44,7 @@ public:
 	 * @param MaterialInfo The desired material info to animate.
 	 */
 	UFUNCTION(BlueprintCallable, Category = "Sequencer|Track", meta = (ScriptMethod))
-	static void SetMaterialInfo(UMovieScenePrimitiveMaterialTrack* Track, const FComponentMaterialInfo& MaterialInfo);
+	static UE_API void SetMaterialInfo(UMovieScenePrimitiveMaterialTrack* Track, const FComponentMaterialInfo& MaterialInfo);
 
 	/**
 	 * Get material info of the component that is animated by the material track.
@@ -50,5 +52,7 @@ public:
 	 * @return The material info.
 	 */
 	UFUNCTION(BlueprintCallable, Category = "Sequencer|Track", meta=(ScriptMethod, DeprecatedFunction, DeprecationMessage = "Use SetMaterialInfo instead"))
-	static FComponentMaterialInfo GetMaterialInfo(UMovieScenePrimitiveMaterialTrack* Track);
+	static UE_API FComponentMaterialInfo GetMaterialInfo(UMovieScenePrimitiveMaterialTrack* Track);
 };
+
+#undef UE_API
diff --git a/Engine/Plugins/MovieScene/SequencerScripting/Source/SequencerScripting/Public/ExtensionLibraries/MovieScenePropertyTrackExtensions.h b/Engine/Plugins/MovieScene/SequencerScripting/Source/SequencerScripting/Public/ExtensionLibraries/MovieScenePropertyTrackExtensions.h
--- a/Engine/Plugins/MovieScene/SequencerScripting/Source/SequencerScripting/Public/ExtensionLibraries/MovieScenePropertyTrackExtensions.h
+++ b/Engine/Plugins/MovieScene/SequencerScripting/Source/SequencerScripting/Public/ExtensionLibraries/MovieScenePropertyTrackExtensions.h
@@ -10,6 +10,8 @@ class UMovieSceneByteTrack;
 class UMovieScenePropertyTrack;
 class UMovieSceneObjectPropertyTrack;
 
+#define UE_API SEQUENCERSCRIPTING_API
+
 /**
  * Function library containing methods that should be hoisted onto UMovieScenePropertyTrack for scripting
  */
@@ -28,7 +30,7 @@ public:
 	 * @param InPropertyPath The property path
 	 */
 	UFUNCTION(BlueprintCallable, Category = "Sequencer|Track", meta = (ScriptMethod))
-	static void SetPropertyNameAndPath(UMovieScenePropertyTrack* Track, const FName& InPropertyName, const FString& InPropertyPath);
+	static UE_API void SetPropertyNameAndPath(UMovieScenePropertyTrack* Track, const FName& InPropertyName, const FString& InPropertyPath);
 
 	/**
 	 * Get this track's property name
@@ -37,7 +39,7 @@ public:
 	 * @return This track's property name
 	 */
 	UFUNCTION(BlueprintCallable, Category = "Sequencer|Track", meta=(ScriptMethod))
-	static FName GetPropertyName(UMovieScenePropertyTrack* Track);
+	static UE_API FName GetPropertyName(UMovieScenePropertyTrack* Track);
 
 	/**
 	 * Get this track's property path
@@ -46,7 +48,7 @@ public:
 	 * @return This track's property path
 	 */
 	UFUNCTION(BlueprintCallable, Category = "Sequencer|Track", meta=(ScriptMethod))
-	static FString GetPropertyPath(UMovieScenePropertyTrack* Track);
+	static UE_API FString GetPropertyPath(UMovieScenePropertyTrack* Track);
 
 	/**
 	 * Get this track's unique name
@@ -55,7 +57,7 @@ public:
 	 * @return This track's unique name
 	 */
 	UFUNCTION(BlueprintCallable, Category = "Sequencer|Track", meta=(ScriptMethod))
-	static FName GetUniqueTrackName(UMovieScenePropertyTrack* Track);
+	static UE_API FName GetUniqueTrackName(UMovieScenePropertyTrack* Track);
 
 	/**
 	 * Set this object property track's property class
@@ -65,7 +67,7 @@ public:
 	 * @param bInClassProperty Whether the property class is a class property or not, which determines whether a class property chooser should be used instead of an object property chooser
 	 */
 	UFUNCTION(BlueprintCallable, Category = "Sequencer|Track", meta = (ScriptMethod))
-	static void SetObjectPropertyClass(UMovieSceneObjectPropertyTrack* Track, UClass* PropertyClass, bool bInClassProperty = false);
+	static UE_API void SetObjectPropertyClass(UMovieSceneObjectPropertyTrack* Track, UClass* PropertyClass, bool bInClassProperty = false);
 
 	/**
 	 * Get this object property track's property class
@@ -74,7 +76,7 @@ public:
 	 * @return The property class for this object property track
 	 */
 	UFUNCTION(BlueprintCallable, Category = "Sequencer|Track", meta = (ScriptMethod))
-	static UClass* GetObjectPropertyClass(UMovieSceneObjectPropertyTrack* Track);
+	static UE_API UClass* GetObjectPropertyClass(UMovieSceneObjectPropertyTrack* Track);
 
 	/**
 	 * Set this byte track's enum
@@ -83,7 +85,7 @@ public:
 	 * @param InEnum The enum to set
 	 */
 	UFUNCTION(BlueprintCallable, Category = "Sequencer|Track", meta = (ScriptMethod))
-	static void SetByteTrackEnum(UMovieSceneByteTrack* Track, UEnum* InEnum);
+	static UE_API void SetByteTrackEnum(UMovieSceneByteTrack* Track, UEnum* InEnum);
 
 	/**
 	 * Get this byte track's enum
@@ -92,5 +94,7 @@ public:
 	 * @return The enum for this byte track
 	 */
 	UFUNCTION(BlueprintPure, Category = "Sequencer|Track", meta = (ScriptMethod))
-	static UEnum* GetByteTrackEnum(UMovieSceneByteTrack* Track);
+	static UE_API UEnum* GetByteTrackEnum(UMovieSceneByteTrack* Track);
 };
+
+#undef UE_API
diff --git a/Engine/Plugins/MovieScene/SequencerScripting/Source/SequencerScripting/Public/ExtensionLibraries/MovieSceneSectionExtensions.h b/Engine/Plugins/MovieScene/SequencerScripting/Source/SequencerScripting/Public/ExtensionLibraries/MovieSceneSectionExtensions.h
--- a/Engine/Plugins/MovieScene/SequencerScripting/Source/SequencerScripting/Public/ExtensionLibraries/MovieSceneSectionExtensions.h
+++ b/Engine/Plugins/MovieScene/SequencerScripting/Source/SequencerScripting/Public/ExtensionLibraries/MovieSceneSectionExtensions.h
@@ -12,6 +12,8 @@ struct FSequencerScriptingRange;
 struct FMovieSceneChannel;
 class UMovieSceneSection;
 
+#define UE_API SEQUENCERSCRIPTING_API
+
 /**
  * Function library containing methods that should be hoisted onto UMovieSceneSections for scripting
  */
@@ -28,7 +30,7 @@ public:
 	 * @return Whether this section has a valid start frame (else infinite)
 	 */
 	UFUNCTION(BlueprintCallable, Category = "Sequencer|Section", meta = (ScriptMethod))
-	static bool HasStartFrame(UMovieSceneSection* Section);
+	static UE_API bool HasStartFrame(UMovieSceneSection* Section);
 
 	/**
 	 * Get start frame. Will throw an exception if section has no start frame, use HasStartFrame to check first.
@@ -37,7 +39,7 @@ public:
 	 * @return Start frame of this section
 	 */
 	UFUNCTION(BlueprintCallable, Category = "Sequencer|Section", meta = (ScriptMethod))
-	static int32 GetStartFrame(UMovieSceneSection* Section);
+	static UE_API int32 GetStartFrame(UMovieSceneSection* Section);
 
 	/**
 	 * Get start time in seconds. Will throw an exception if section has no start frame, use HasStartFrame to check first.
@@ -46,7 +48,7 @@ public:
 	 * @return Start time of this section
 	 */
 	UFUNCTION(BlueprintCallable, Category = "Sequencer|Section", meta = (ScriptMethod))
-	static float GetStartFrameSeconds(UMovieSceneSection* Section);
+	static UE_API float GetStartFrameSeconds(UMovieSceneSection* Section);
 
 	/**
 	 * Has end frame
@@ -55,7 +57,7 @@ public:
 	 * @return Whether this section has a valid end frame (else infinite)
 	 */
 	UFUNCTION(BlueprintCallable, Category = "Sequencer|Section", meta = (ScriptMethod))
-	static bool HasEndFrame(UMovieSceneSection* Section);
+	static UE_API bool HasEndFrame(UMovieSceneSection* Section);
 
 	/**
 	 * Get end frame. Will throw an exception if section has no end frame, use HasEndFrame to check first.
@@ -64,7 +66,7 @@ public:
 	 * @return End frame of this section
 	 */
 	UFUNCTION(BlueprintCallable, Category = "Sequencer|Section", meta = (ScriptMethod))
-	static int32 GetEndFrame(UMovieSceneSection* Section);
+	static UE_API int32 GetEndFrame(UMovieSceneSection* Section);
 
 	/**
 	 * Get end time in seconds. Will throw an exception if section has no end frame, use HasEndFrame to check first.
@@ -73,7 +75,7 @@ public:
 	 * @return End time of this section
 	 */
 	UFUNCTION(BlueprintCallable, Category = "Sequencer|Section", meta = (ScriptMethod))
-	static float GetEndFrameSeconds(UMovieSceneSection* Section);
+	static UE_API float GetEndFrameSeconds(UMovieSceneSection* Section);
 
 	/**
 	 * Checks to see if this section has an AutoSize implementation, and if so, if that implementation has a start frame.
@@ -82,7 +84,7 @@ public:
 	 * @return Whether this section has a valid autosize range, and a valid start frame
 	 */
 	UFUNCTION(BlueprintCallable, Category = "Sequencer|Section", meta = (ScriptMethod))
-	static bool GetAutoSizeHasStartFrame(UMovieSceneSection* Section);
+	static UE_API bool GetAutoSizeHasStartFrame(UMovieSceneSection* Section);
 
 	/**
 	 * Get start frame of the AutoSize. Will throw an exception if section has no start frame, use GetAutoSizeHasStartFrame to check first.
@@ -91,7 +93,7 @@ public:
 	 * @return The start frame of the AutoSize data.
 	 */
 	UFUNCTION(BlueprintCallable, Category = "Sequencer|Section", meta = (ScriptMethod))
-	static int32 GetAutoSizeStartFrame(UMovieSceneSection* Section);
+	static UE_API int32 GetAutoSizeStartFrame(UMovieSceneSection* Section);
 
 	/**
 	 * Get start time of the AutoSize in seconds. Will throw an exception if section has no start frame, use GetAutoSizeHasStartFrame to check first.
@@ -100,7 +102,7 @@ public:
 	 * @return The start frame of the AutoSize data in seconds.
 	 */
 	UFUNCTION(BlueprintCallable, Category = "Sequencer|Section", meta = (ScriptMethod))
-	static float GetAutoSizeStartFrameSeconds(UMovieSceneSection* Section);
+	static UE_API float GetAutoSizeStartFrameSeconds(UMovieSceneSection* Section);
 
 	/**
 	 * Checks to see if this section has an AutoSize implementation, and if so, if that implementation has a end frame.
@@ -109,7 +111,7 @@ public:
 	 * @return Whether this section has a valid autosize range, and a valid end frame
 	 */
 	UFUNCTION(BlueprintCallable, Category = "Sequencer|Section", meta = (ScriptMethod))
-	static bool GetAutoSizeHasEndFrame(UMovieSceneSection* Section);
+	static UE_API bool GetAutoSizeHasEndFrame(UMovieSceneSection* Section);
 
 	/**
 	 * Get end frame of the AutoSize. Will throw an exception if section has no end frame, use GetAutoSizeHasEndFrame to check first.
@@ -118,7 +120,7 @@ public:
 	 * @return The end frame of the AutoSize data.
 	 */
 	UFUNCTION(BlueprintCallable, Category = "Sequencer|Section", meta = (ScriptMethod))
-	static int32 GetAutoSizeEndFrame(UMovieSceneSection* Section);
+	static UE_API int32 GetAutoSizeEndFrame(UMovieSceneSection* Section);
 
 	/**
 	 * Get end time of the AutoSize seconds. Will throw an exception if section has no end frame, use GetAutoSizeHasEndFrame to check first.
@@ -127,7 +129,7 @@ public:
 	 * @return The end frame of the AutoSize data in seconds.
 	 */
 	UFUNCTION(BlueprintCallable, Category = "Sequencer|Section", meta = (ScriptMethod))
-	static float GetAutoSizeEndFrameSeconds(UMovieSceneSection* Section);
+	static UE_API float GetAutoSizeEndFrameSeconds(UMovieSceneSection* Section);
 
 	/**
 	 * Set range
@@ -137,7 +139,7 @@ public:
 	 * @param EndFrame The desired end frame for this section
 	 */
 	UFUNCTION(BlueprintCallable, Category = "Sequencer|Section", meta = (ScriptMethod))
-	static void SetRange(UMovieSceneSection* Section, int32 StartFrame, int32 EndFrame);
+	static UE_API void SetRange(UMovieSceneSection* Section, int32 StartFrame, int32 EndFrame);
 
 	/**
 	 * Set range in seconds
@@ -147,7 +149,7 @@ public:
 	 * @param EndTime The desired end frame for this section
 	 */
 	UFUNCTION(BlueprintCallable, Category = "Sequencer|Section", meta = (ScriptMethod))
-	static void SetRangeSeconds(UMovieSceneSection* Section, float StartTime, float EndTime);
+	static UE_API void SetRangeSeconds(UMovieSceneSection* Section, float StartTime, float EndTime);
 
 	/**
 	 * Set start frame
@@ -156,7 +158,7 @@ public:
 	 * @param StartFrame The desired start frame for this section
 	 */
 	UFUNCTION(BlueprintCallable, Category = "Sequencer|Section", meta = (ScriptMethod))
-	static void SetStartFrame(UMovieSceneSection* Section, int32 StartFrame);
+	static UE_API void SetStartFrame(UMovieSceneSection* Section, int32 StartFrame);
 
 	/**
 	 * Set start time in seconds
@@ -165,7 +167,7 @@ public:
 	 * @param StartTime The desired start time for this section
 	 */
 	UFUNCTION(BlueprintCallable, Category = "Sequencer|Section", meta = (ScriptMethod))
-	static void SetStartFrameSeconds(UMovieSceneSection* Section, float StartTime);
+	static UE_API void SetStartFrameSeconds(UMovieSceneSection* Section, float StartTime);
 
 	/**
 	 * Set start frame bounded
@@ -174,7 +176,7 @@ public:
 	 * @param IsBounded The desired bounded state of the start frame
 	 */
 	UFUNCTION(BlueprintCallable, Category = "Sequencer|Section", meta = (ScriptMethod))
-	static void SetStartFrameBounded(UMovieSceneSection* Section, bool bIsBounded);
+	static UE_API void SetStartFrameBounded(UMovieSceneSection* Section, bool bIsBounded);
 
 	/**
 	 * Set end frame
@@ -183,7 +185,7 @@ public:
 	 * @param EndFrame The desired start frame for this section
 	 */
 	UFUNCTION(BlueprintCallable, Category = "Sequencer|Section", meta = (ScriptMethod))
-	static void SetEndFrame(UMovieSceneSection* Section, int32 EndFrame);
+	static UE_API void SetEndFrame(UMovieSceneSection* Section, int32 EndFrame);
 
 	/**
 	 * Set end time in seconds
@@ -192,7 +194,7 @@ public:
 	 * @param EndTime The desired end time for this section
 	 */
 	UFUNCTION(BlueprintCallable, Category = "Sequencer|Section", meta = (ScriptMethod))
-	static void SetEndFrameSeconds(UMovieSceneSection* Section, float EndTime);
+	static UE_API void SetEndFrameSeconds(UMovieSceneSection* Section, float EndTime);
 
 	/**
      * Set end frame bounded
@@ -201,7 +203,7 @@ public:
 	 * @param IsBounded The desired bounded state of the end frame
 	 */
 	UFUNCTION(BlueprintCallable, Category = "Sequencer|Section", meta = (ScriptMethod))
-	static void SetEndFrameBounded(UMovieSceneSection* Section, bool bIsBounded);
+	static UE_API void SetEndFrameBounded(UMovieSceneSection* Section, bool bIsBounded);
 
 	/**
 	* Find all channels that belong to the specified UMovieSceneSection. Some sections have many channels (such as
@@ -212,7 +214,7 @@ public:
 	* @return An array containing any key channels that match the type specified
 	*/
 	UFUNCTION(BlueprintCallable, Category = "Sequencer|Section", meta=(ScriptMethod, DevelopmentOnly))
-	static TArray<UMovieSceneScriptingChannel*> GetAllChannels(UMovieSceneSection* Section);
+	static UE_API TArray<UMovieSceneScriptingChannel*> GetAllChannels(UMovieSceneSection* Section);
 
 	/**
 	* Find all channels that belong to the specified UMovieSceneSection that match the specific type. This will filter out any children who do not inherit
@@ -223,7 +225,7 @@ public:
 	* @return An array containing any key channels that match the type specified
 	*/
 	UFUNCTION(BlueprintCallable, Category = "Sequencer|Section", meta = (ScriptMethod, DeterminesOutputType="TrackType", DevelopmentOnly))
-	static TArray<UMovieSceneScriptingChannel*> GetChannelsByType(UMovieSceneSection* Section, TSubclassOf<UMovieSceneScriptingChannel> ChannelType);
+	static UE_API TArray<UMovieSceneScriptingChannel*> GetChannelsByType(UMovieSceneSection* Section, TSubclassOf<UMovieSceneScriptingChannel> ChannelType);
 
 
 
@@ -235,7 +237,7 @@ public:
 	* @return The channel if it exists 
 	*/
 	UFUNCTION(BlueprintCallable, Category = "Sequencer|Section", meta = (ScriptMethod, DeterminesOutputType = "TrackType", DevelopmentOnly))
-	static UMovieSceneScriptingChannel* GetChannel(UMovieSceneSection* Section, const FName& ChannelName);
+	static UE_API UMovieSceneScriptingChannel* GetChannel(UMovieSceneSection* Section, const FName& ChannelName);
 
 
 	/**
@@ -247,13 +249,15 @@ public:
 	 * @return The frame at the parent sequence
 	 */
 	UFUNCTION(BlueprintCallable, Category = "Sequencer|Section", meta = (ScriptMethod))
-	static int32 GetParentSequenceFrame(UMovieSceneSubSection* Section, int32 InFrame, UMovieSceneSequence* ParentSequence);
+	static UE_API int32 GetParentSequenceFrame(UMovieSceneSubSection* Section, int32 InFrame, UMovieSceneSequence* ParentSequence);
 
 public:
 	/**
 	* Utility function to get channel from a section and a name
 	*/
-	static SEQUENCERSCRIPTING_API FMovieSceneChannel* GetMovieSceneChannel(UMovieSceneSection* InSection, const FName& InName);
+	static UE_API FMovieSceneChannel* GetMovieSceneChannel(UMovieSceneSection* InSection, const FName& InName);
 
 
 };
+
+#undef UE_API
diff --git a/Engine/Plugins/MovieScene/SequencerScripting/Source/SequencerScripting/Public/ExtensionLibraries/MovieSceneSequenceExtensions.h b/Engine/Plugins/MovieScene/SequencerScripting/Source/SequencerScripting/Public/ExtensionLibraries/MovieSceneSequenceExtensions.h
--- a/Engine/Plugins/MovieScene/SequencerScripting/Source/SequencerScripting/Public/ExtensionLibraries/MovieSceneSequenceExtensions.h
+++ b/Engine/Plugins/MovieScene/SequencerScripting/Source/SequencerScripting/Public/ExtensionLibraries/MovieSceneSequenceExtensions.h
@@ -20,6 +20,8 @@ struct FSequencerScriptingRange;
 struct FTimecode;
 template <typename T> class TSubclassOf;
 
+#define UE_API SEQUENCERSCRIPTING_API
+
 /**
  * Function library containing methods that should be hoisted onto UMovieSceneSequences for scripting purposes
  */
@@ -35,7 +37,7 @@ public:
 	 * @return This sequence's movie scene data object
 	 */
 	UFUNCTION(BlueprintCallable, Category="Sequencer|Sequence", meta=(ScriptMethod))
-	static UMovieScene* GetMovieScene(UMovieSceneSequence* Sequence);
+	static UE_API UMovieScene* GetMovieScene(UMovieSceneSequence* Sequence);
 
 	/**
 	 * Get all tracks
@@ -44,7 +46,7 @@ public:
 	 * @return An array containing all tracks in this sequence
 	 */
 	UFUNCTION(BlueprintCallable, Category="Sequencer|Sequence", meta=(ScriptMethod))
-	static TArray<UMovieSceneTrack*> GetTracks(UMovieSceneSequence* Sequence);
+	static UE_API TArray<UMovieSceneTrack*> GetTracks(UMovieSceneSequence* Sequence);
 
 	/**
 	 * Find all tracks of the specified type
@@ -54,7 +56,7 @@ public:
 	 * @return An array containing any tracks that match the type specified
 	 */
 	UFUNCTION(BlueprintCallable, Category = "Sequencer|Sequence", meta = (ScriptMethod, DeterminesOutputType = "TrackType"))
-	static TArray<UMovieSceneTrack*> FindTracksByType(UMovieSceneSequence* Sequence, TSubclassOf<UMovieSceneTrack> TrackType);
+	static UE_API TArray<UMovieSceneTrack*> FindTracksByType(UMovieSceneSequence* Sequence, TSubclassOf<UMovieSceneTrack> TrackType);
 
 	/**
 	 * Find all tracks of the specified type, not allowing sub-classed types
@@ -64,7 +66,7 @@ public:
 	 * @return An array containing any tracks that are exactly the same as the type specified
 	 */
 	UFUNCTION(BlueprintCallable, Category="Sequencer|Sequence", meta=(ScriptMethod, DeterminesOutputType="TrackType"))
-	static TArray<UMovieSceneTrack*> FindTracksByExactType(UMovieSceneSequence* Sequence, TSubclassOf<UMovieSceneTrack> TrackType);
+	static UE_API TArray<UMovieSceneTrack*> FindTracksByExactType(UMovieSceneSequence* Sequence, TSubclassOf<UMovieSceneTrack> TrackType);
 
 	/**
 	 * Add a new track of the specified type
@@ -74,7 +76,7 @@ public:
 	 * @return The newly created track, if successful
 	 */
 	UFUNCTION(BlueprintCallable, Category = "Sequencer|Sequence", meta = (ScriptMethod, DeterminesOutputType = "TrackType"))
-	static UMovieSceneTrack* AddTrack(UMovieSceneSequence* Sequence, TSubclassOf<UMovieSceneTrack> TrackType);
+	static UE_API UMovieSceneTrack* AddTrack(UMovieSceneSequence* Sequence, TSubclassOf<UMovieSceneTrack> TrackType);
 
 	/**
 	 * Removes a track
@@ -84,7 +86,7 @@ public:
 	 * @return Whether the track was successfully removed
 	 */
 	UFUNCTION(BlueprintCallable, Category="Sequencer|Sequence", meta=(ScriptMethod))
-	static bool RemoveTrack(UMovieSceneSequence* Sequence, UMovieSceneTrack* Track);
+	static UE_API bool RemoveTrack(UMovieSceneSequence* Sequence, UMovieSceneTrack* Track);
 
 	/**
 	 * Gets this sequence's display rate
@@ -93,7 +95,7 @@ public:
 	 * @return The display rate that this sequence is displayed as
 	 */
 	UFUNCTION(BlueprintCallable, Category="Sequencer|Sequence", meta=(ScriptMethod))
-	static FFrameRate GetDisplayRate(UMovieSceneSequence* Sequence);
+	static UE_API FFrameRate GetDisplayRate(UMovieSceneSequence* Sequence);
 
 	/**
 	 * Sets this sequence's display rate
@@ -102,7 +104,7 @@ public:
 	 * @param DisplayRate The display rate that this sequence is displayed as
 	 */
 	UFUNCTION(BlueprintCallable, Category="Sequencer|Sequence", meta=(ScriptMethod))
-	static void SetDisplayRate(UMovieSceneSequence* Sequence, FFrameRate DisplayRate);
+	static UE_API void SetDisplayRate(UMovieSceneSequence* Sequence, FFrameRate DisplayRate);
 
 	/**
 	 * Gets this sequence's tick resolution
@@ -111,7 +113,7 @@ public:
 	 * @return The tick resolution of the sequence, defining the smallest unit of time representable on this sequence
 	 */
 	UFUNCTION(BlueprintCallable, Category="Sequencer|Sequence", meta=(ScriptMethod))
-	static FFrameRate GetTickResolution(UMovieSceneSequence* Sequence);
+	static UE_API FFrameRate GetTickResolution(UMovieSceneSequence* Sequence);
 
 	/**
 	 * Sets this sequence's tick resolution and migrates frame times
@@ -120,7 +122,7 @@ public:
 	 * @param TickResolution The tick resolution of the sequence, defining the smallest unit of time representable on this sequence
 	 */
 	UFUNCTION(BlueprintCallable, Category="Sequencer|Sequence", meta=(ScriptMethod))
-	static void SetTickResolution(UMovieSceneSequence* Sequence, FFrameRate TickResolution);
+	static UE_API void SetTickResolution(UMovieSceneSequence* Sequence, FFrameRate TickResolution);
 
 	/**
 	 * Sets this sequence's tick resolution directly without migrating frame times
@@ -129,7 +131,7 @@ public:
 	 * @param TickResolution The tick resolution of the sequence, defining the smallest unit of time representable on this sequence
 	 */
 	UFUNCTION(BlueprintCallable, Category="Sequencer|Sequence", meta=(ScriptMethod))
-	static void SetTickResolutionDirectly(UMovieSceneSequence* Sequence, FFrameRate TickResolution);
+	static UE_API void SetTickResolutionDirectly(UMovieSceneSequence* Sequence, FFrameRate TickResolution);
 
 	/**
 	 * Make a new range for this sequence in its display rate
@@ -140,7 +142,7 @@ public:
 	 * @return Specified sequencer range
 	 */
 	UFUNCTION(BlueprintCallable, Category="Sequencer|Sequence", meta=(ScriptMethod))
-	static FSequencerScriptingRange MakeRange(UMovieSceneSequence* Sequence, int32 StartFrame, int32 Duration);
+	static UE_API FSequencerScriptingRange MakeRange(UMovieSceneSequence* Sequence, int32 StartFrame, int32 Duration);
 
 	/**
 	 * Make a new range for this sequence in seconds
@@ -151,7 +153,7 @@ public:
 	 * @return Specified sequencer range
 	 */
 	UFUNCTION(BlueprintCallable, Category="Sequencer|Sequence", meta=(ScriptMethod))
-	static FSequencerScriptingRange MakeRangeSeconds(UMovieSceneSequence* Sequence, float StartTime, float Duration);
+	static UE_API FSequencerScriptingRange MakeRangeSeconds(UMovieSceneSequence* Sequence, float StartTime, float Duration);
 
 	/**
 	 * Get playback range of this sequence in display rate resolution
@@ -160,7 +162,7 @@ public:
 	 * @return Playback range of this sequence
 	 */
 	UFUNCTION(BlueprintCallable, Category = "Sequencer|Sequence", meta = (ScriptMethod))
-	static FSequencerScriptingRange GetPlaybackRange(UMovieSceneSequence* Sequence);
+	static UE_API FSequencerScriptingRange GetPlaybackRange(UMovieSceneSequence* Sequence);
 
 	/**
 	 * Get playback start of this sequence in display rate resolution
@@ -169,7 +171,7 @@ public:
 	 * @return Playback start of this sequence
 	 */
 	UFUNCTION(BlueprintCallable, Category = "Sequencer|Sequence", meta = (ScriptMethod))
-	static int32 GetPlaybackStart(UMovieSceneSequence* Sequence);
+	static UE_API int32 GetPlaybackStart(UMovieSceneSequence* Sequence);
 
 	/**
 	 * Get playback start of this sequence in seconds
@@ -178,7 +180,7 @@ public:
 	 * @return Playback start of this sequence
 	 */
 	UFUNCTION(BlueprintCallable, Category = "Sequencer|Sequence", meta = (ScriptMethod))
-	static float GetPlaybackStartSeconds(UMovieSceneSequence* Sequence);
+	static UE_API float GetPlaybackStartSeconds(UMovieSceneSequence* Sequence);
 
 	/**
 	 * Get playback end of this sequence in display rate resolution
@@ -187,7 +189,7 @@ public:
 	 * @return Playback end of this sequence
 	 */
 	UFUNCTION(BlueprintCallable, Category = "Sequencer|Sequence", meta = (ScriptMethod))
-	static int32 GetPlaybackEnd(UMovieSceneSequence* Sequence);
+	static UE_API int32 GetPlaybackEnd(UMovieSceneSequence* Sequence);
 
 	/**
 	 * Get playback end of this sequence in seconds
@@ -196,7 +198,7 @@ public:
 	 * @return Playback end of this sequence
 	 */
 	UFUNCTION(BlueprintCallable, Category = "Sequencer|Sequence", meta = (ScriptMethod))
-	static float GetPlaybackEndSeconds(UMovieSceneSequence* Sequence);
+	static UE_API float GetPlaybackEndSeconds(UMovieSceneSequence* Sequence);
 
 	/**
 	 * Set playback start of this sequence
@@ -205,7 +207,7 @@ public:
 	 * @param StartFrame      The desired start frame for this sequence
 	 */
 	UFUNCTION(BlueprintCallable, Category = "Sequencer|Sequence", meta = (ScriptMethod))
-	static void SetPlaybackStart(UMovieSceneSequence* Sequence, int32 StartFrame);
+	static UE_API void SetPlaybackStart(UMovieSceneSequence* Sequence, int32 StartFrame);
 
 	/**
 	 * Set playback start of this sequence in seconds
@@ -214,7 +216,7 @@ public:
 	 * @param StartTime       The desired start time in seconds for this sequence
 	 */
 	UFUNCTION(BlueprintCallable, Category = "Sequencer|Sequence", meta = (ScriptMethod))
-	static void SetPlaybackStartSeconds(UMovieSceneSequence* Sequence, float StartTime);
+	static UE_API void SetPlaybackStartSeconds(UMovieSceneSequence* Sequence, float StartTime);
 
 	/**
 	 * Set playback end of this sequence
@@ -223,7 +225,7 @@ public:
 	 * @param EndFrame        The desired end frame for this sequence
 	 */
 	UFUNCTION(BlueprintCallable, Category = "Sequencer|Sequence", meta = (ScriptMethod))
-	static void SetPlaybackEnd(UMovieSceneSequence* Sequence, int32 EndFrame);
+	static UE_API void SetPlaybackEnd(UMovieSceneSequence* Sequence, int32 EndFrame);
 
 	/**
 	 * Set playback end of this sequence in seconds
@@ -232,7 +234,7 @@ public:
 	 * @param EndTime         The desired end time in seconds for this sequence
 	 */
 	UFUNCTION(BlueprintCallable, Category = "Sequencer|Sequence", meta = (ScriptMethod))
-	static void SetPlaybackEndSeconds(UMovieSceneSequence* Sequence, float EndTime);
+	static UE_API void SetPlaybackEndSeconds(UMovieSceneSequence* Sequence, float EndTime);
 
 	/**
 	 * Set the sequence view range start in seconds
@@ -241,7 +243,7 @@ public:
 	 * @param StartTimeInSeconds The desired view range start time in seconds for this sequence
 	 */
 	UFUNCTION(BlueprintCallable, Category = "Sequencer|Sequence", meta = (ScriptMethod, DevelopmentOnly))
-	static void SetViewRangeStart(UMovieSceneSequence* InSequence, double StartTimeInSeconds);
+	static UE_API void SetViewRangeStart(UMovieSceneSequence* InSequence, double StartTimeInSeconds);
 
 	/**
 	 * Get the sequence view range start in seconds
@@ -250,7 +252,7 @@ public:
 	 * @return The view range start time in seconds for this sequence
 	 */
 	UFUNCTION(BlueprintPure, Category = "Sequencer|Sequence", meta = (ScriptMethod, DevelopmentOnly))
-	static double GetViewRangeStart(UMovieSceneSequence* InSequence);
+	static UE_API double GetViewRangeStart(UMovieSceneSequence* InSequence);
 
 	/**
 	 * Set the sequence view range end in seconds
@@ -259,7 +261,7 @@ public:
 	 * @param StartTimeInSeconds The desired view range end time in seconds for this sequence
 	 */
 	UFUNCTION(BlueprintCallable, Category = "Sequencer|Sequence", meta = (ScriptMethod, DevelopmentOnly))
-	static void SetViewRangeEnd(UMovieSceneSequence* InSequence, double EndTimeInSeconds);
+	static UE_API void SetViewRangeEnd(UMovieSceneSequence* InSequence, double EndTimeInSeconds);
 
 	/**
 	 * Get the sequence view range end in seconds
@@ -268,7 +270,7 @@ public:
 	 * @return The view range end time in seconds for this sequence
 	 */
 	UFUNCTION(BlueprintPure, Category = "Sequencer|Sequence", meta = (ScriptMethod, DevelopmentOnly))
-	static double GetViewRangeEnd(UMovieSceneSequence* InSequence);
+	static UE_API double GetViewRangeEnd(UMovieSceneSequence* InSequence);
 
 	/**
 	 * Set the sequence work range start in seconds
@@ -277,7 +279,7 @@ public:
 	 * @param StartTimeInSeconds The desired work range start time in seconds for this sequence
 	 */
 	UFUNCTION(BlueprintCallable, Category = "Sequencer|Sequence", meta = (ScriptMethod, DevelopmentOnly))
-	static void SetWorkRangeStart(UMovieSceneSequence* InSequence, double StartTimeInSeconds);
+	static UE_API void SetWorkRangeStart(UMovieSceneSequence* InSequence, double StartTimeInSeconds);
 
 	/**
 	 * Get the sequence work range start in seconds
@@ -286,7 +288,7 @@ public:
 	 * @return The work range start time in seconds for this sequence
 	 */
 	UFUNCTION(BlueprintPure, Category = "Sequencer|Sequence", meta = (ScriptMethod, DevelopmentOnly))
-	static double GetWorkRangeStart(UMovieSceneSequence* InSequence);
+	static UE_API double GetWorkRangeStart(UMovieSceneSequence* InSequence);
 
 	/**
 	 * Set the sequence work range end in seconds
@@ -295,7 +297,7 @@ public:
 	 * @param StartTimeInSeconds The desired work range end time in seconds for this sequence
 	 */
 	UFUNCTION(BlueprintCallable, Category = "Sequencer|Sequence", meta = (ScriptMethod, DevelopmentOnly))
-	static void SetWorkRangeEnd(UMovieSceneSequence* InSequence, double EndTimeInSeconds);
+	static UE_API void SetWorkRangeEnd(UMovieSceneSequence* InSequence, double EndTimeInSeconds);
 
 	/**
 	 * Get the sequence work range end in seconds
@@ -304,7 +306,7 @@ public:
 	 * @return The work range end time in seconds for this sequence
 	 */
 	UFUNCTION(BlueprintPure, Category = "Sequencer|Sequence", meta = (ScriptMethod, DevelopmentOnly))
-	static double GetWorkRangeEnd(UMovieSceneSequence* InSequence);
+	static UE_API double GetWorkRangeEnd(UMovieSceneSequence* InSequence);
 
 	/**
 	 * Set the evaluation type for this sequence
@@ -313,7 +315,7 @@ public:
 	 * @param InEvaluationType The evaluation type to set for this sequence
 	 */
 	UFUNCTION(BlueprintCallable, Category = "Sequencer|Sequence", meta = (ScriptMethod))
-	static void SetEvaluationType(UMovieSceneSequence* InSequence, EMovieSceneEvaluationType InEvaluationType);
+	static UE_API void SetEvaluationType(UMovieSceneSequence* InSequence, EMovieSceneEvaluationType InEvaluationType);
 
 	/**
 	 * Get the evaluation type for this sequence
@@ -322,7 +324,7 @@ public:
 	 * @return The evaluation type for this sequence
 	 */
 	UFUNCTION(BlueprintPure, Category = "Sequencer|Sequence", meta = (ScriptMethod))
-	static EMovieSceneEvaluationType GetEvaluationType(UMovieSceneSequence* InSequence);
+	static UE_API EMovieSceneEvaluationType GetEvaluationType(UMovieSceneSequence* InSequence);
 
 	/**
 	 * Set the clock source for this sequence
@@ -331,7 +333,7 @@ public:
 	 * @param InClockSource The clock source to set for this sequence
 	 */
 	UFUNCTION(BlueprintCallable, Category = "Sequencer|Sequence", meta = (ScriptMethod))
-	static void SetClockSource(UMovieSceneSequence* InSequence, EUpdateClockSource InClockSource);
+	static UE_API void SetClockSource(UMovieSceneSequence* InSequence, EUpdateClockSource InClockSource);
 
 	/**
 	 * Get the clock source for this sequence
@@ -340,7 +342,7 @@ public:
 	 * @return The clock source for this sequence
 	 */
 	UFUNCTION(BlueprintPure, Category = "Sequencer|Sequence", meta = (ScriptMethod))
-	static EUpdateClockSource GetClockSource(UMovieSceneSequence* InSequence);
+	static UE_API EUpdateClockSource GetClockSource(UMovieSceneSequence* InSequence);
 
 	/**
 	 * Get the custom clock for this sequence
@@ -349,7 +351,7 @@ public:
 	 * @return The MovieSceneClock for this sequence
 	 */
 	UFUNCTION(BlueprintPure, Category = "Sequencer|Sequence", meta = (ScriptMethod))
-	static UMovieSceneClock* GetCustomClock(UMovieSceneSequence* InSequence);
+	static UE_API UMovieSceneClock* GetCustomClock(UMovieSceneSequence* InSequence);
 
 	/**
 	 * Attempt to locate a binding in this sequence by its name
@@ -359,7 +361,7 @@ public:
 	 * @return A unique identifier for the binding, or invalid
 	 */
 	UFUNCTION(BlueprintCallable, Category="Sequencer|Sequence", meta=(ScriptMethod))
-	static FMovieSceneBindingProxy FindBindingByName(UMovieSceneSequence* Sequence, FString Name);
+	static UE_API FMovieSceneBindingProxy FindBindingByName(UMovieSceneSequence* Sequence, FString Name);
 
 	/**
 	 * Attempt to locate a binding in this sequence by its Id
@@ -369,7 +371,7 @@ public:
 	 * @return A unique identifier for the binding, or invalid
 	 */
 	UFUNCTION(BlueprintCallable, Category="Sequencer|Sequence", meta=(ScriptMethod))
-	static FMovieSceneBindingProxy FindBindingById(UMovieSceneSequence* Sequence, FGuid BindingId);
+	static UE_API FMovieSceneBindingProxy FindBindingById(UMovieSceneSequence* Sequence, FGuid BindingId);
 
 	/**
 	 * Get all the bindings in this sequence
@@ -378,7 +380,7 @@ public:
 	 * @return An array of unique identifiers for all the bindings in this sequence
 	 */
 	UFUNCTION(BlueprintCallable, Category="Sequencer|Sequence", meta=(ScriptMethod))
-	static TArray<FMovieSceneBindingProxy> GetBindings(UMovieSceneSequence* Sequence);
+	static UE_API TArray<FMovieSceneBindingProxy> GetBindings(UMovieSceneSequence* Sequence);
 
 	/**
 	* Get all the spawnables in this sequence. For Level Sequences, this includes bindings with binding type UMovieSceneSpawnableActorBinding.
@@ -387,7 +389,7 @@ public:
 	* @return Spawnables in this sequence
 	*/
 	UFUNCTION(BlueprintCallable, Category="Sequencer|Sequence", meta=(ScriptMethod))
-	static TArray<FMovieSceneBindingProxy> GetSpawnables(UMovieSceneSequence* Sequence);
+	static UE_API TArray<FMovieSceneBindingProxy> GetSpawnables(UMovieSceneSequence* Sequence);
 
 	/**
 	* Get all the possessables in this sequence. It is understood for the purpose of this function that this means the bindings are not custom.
@@ -396,7 +398,7 @@ public:
 	* @return Possessables in this sequence
 	*/
 	UFUNCTION(BlueprintCallable, Category="Sequencer|Sequence", meta=(ScriptMethod))
-	static TArray<FMovieSceneBindingProxy> GetPossessables(UMovieSceneSequence* Sequence);
+	static UE_API TArray<FMovieSceneBindingProxy> GetPossessables(UMovieSceneSequence* Sequence);
 
 	/**
 	 * Add a new binding to this sequence that will possess the specified object
@@ -406,7 +408,7 @@ public:
 	 * @return A unique identifier for the new binding
 	 */
 	UFUNCTION(BlueprintCallable, Category="Sequencer|Sequence", meta=(ScriptMethod))
-	static FMovieSceneBindingProxy AddPossessable(UMovieSceneSequence* Sequence, UObject* ObjectToPossess);
+	static UE_API FMovieSceneBindingProxy AddPossessable(UMovieSceneSequence* Sequence, UObject* ObjectToPossess);
 
 	/**
 	 * Add a new binding to this sequence that will spawn the specified object.
@@ -441,7 +443,7 @@ public:
 	 * @return An array of all bound objects
 	 */
 	UFUNCTION(BlueprintCallable, Category="Sequencer|Sequence", meta=(ScriptMethod))
-	static TArray<UObject*> LocateBoundObjects(UMovieSceneSequence* Sequence, const FMovieSceneBindingProxy& InBinding, UObject* Context);
+	static UE_API TArray<UObject*> LocateBoundObjects(UMovieSceneSequence* Sequence, const FMovieSceneBindingProxy& InBinding, UObject* Context);
 
 
 	/**
@@ -452,7 +454,7 @@ public:
 	 * @return The binding's id
 	 */
 	UFUNCTION(BlueprintCallable, Category = "Sequencer|Sequence", meta = (ScriptMethod))
-	static FMovieSceneObjectBindingID GetBindingID(UMovieSceneSequence* Sequence, const FMovieSceneBindingProxy& InBinding);
+	static UE_API FMovieSceneObjectBindingID GetBindingID(UMovieSceneSequence* Sequence, const FMovieSceneBindingProxy& InBinding);
 
 
 	/**
@@ -466,7 +468,7 @@ public:
 	 * @return The binding's id
 	 */
 	UFUNCTION(BlueprintCallable, Category = "Sequencer|Sequence", meta = (ScriptMethod))
-	static FMovieSceneObjectBindingID GetPortableBindingID(UMovieSceneSequence* RootSequence, UMovieSceneSequence* DestinationSequence, const FMovieSceneBindingProxy& InBinding);
+	static UE_API FMovieSceneObjectBindingID GetPortableBindingID(UMovieSceneSequence* RootSequence, UMovieSceneSequence* DestinationSequence, const FMovieSceneBindingProxy& InBinding);
 
 	/**
 	 * Make a binding for the given binding ID
@@ -476,7 +478,7 @@ public:
 	 * @return The new binding proxy
 	 */
 	UFUNCTION(BlueprintCallable, Category = "Sequencer|Sequence", meta = (ScriptMethod))
-	static FMovieSceneBindingProxy ResolveBindingID(UMovieSceneSequence* RootSequence, FMovieSceneObjectBindingID InObjectBindingID);
+	static UE_API FMovieSceneBindingProxy ResolveBindingID(UMovieSceneSequence* RootSequence, FMovieSceneObjectBindingID InObjectBindingID);
 
 
 	/**
@@ -486,7 +488,7 @@ public:
 	 * @return The folders contained within the given sequence
 	 */
 	UFUNCTION(BlueprintCallable, BlueprintPure, Category="Sequencer|Sequence", meta=(ScriptMethod))
-	static TArray<UMovieSceneFolder*> GetRootFoldersInSequence(UMovieSceneSequence* Sequence);
+	static UE_API TArray<UMovieSceneFolder*> GetRootFoldersInSequence(UMovieSceneSequence* Sequence);
 
 	/**
 	 * Add a root folder to the given sequence
@@ -496,7 +498,7 @@ public:
 	 * @return The newly created folder
 	 */
 	UFUNCTION(BlueprintCallable, Category="Sequencer|Sequence", meta=(ScriptMethod))
-	static UMovieSceneFolder* AddRootFolderToSequence(UMovieSceneSequence* Sequence, FString NewFolderName);
+	static UE_API UMovieSceneFolder* AddRootFolderToSequence(UMovieSceneSequence* Sequence, FString NewFolderName);
 
 	/**
 	 * Remove a root folder from the given sequence. Will throw an exception if the specified folder is not valid or not a root folder.
@@ -505,7 +507,7 @@ public:
 	 * @param Folder			The folder to remove
 	 */
 	UFUNCTION(BlueprintCallable, Category="Sequencer|Sequence", meta=(ScriptMethod))
-	static void RemoveRootFolderFromSequence(UMovieSceneSequence* Sequence, UMovieSceneFolder* Folder);
+	static UE_API void RemoveRootFolderFromSequence(UMovieSceneSequence* Sequence, UMovieSceneFolder* Folder);
 
 public:
 
@@ -517,7 +519,7 @@ public:
 	 * @param bExactMatch    Whether to match sub classes or not
 	 * @return A filtered array of tracks
 	 */
-	static TArray<UMovieSceneTrack*> FilterTracks(TArrayView<UMovieSceneTrack* const> InTracks, UClass* DesiredClass, bool bExactMatch);
+	static UE_API TArray<UMovieSceneTrack*> FilterTracks(TArrayView<UMovieSceneTrack* const> InTracks, UClass* DesiredClass, bool bExactMatch);
 
 public:
 
@@ -527,7 +529,7 @@ public:
 	 * @return Return the user marked frames
 	 */
 	UFUNCTION(BlueprintCallable, Category = "Sequencer|Sequence", meta = (ScriptMethod, DisplayName = "Get Marked Frames"))
-	static TArray<FMovieSceneMarkedFrame> GetMarkedFramesFromSequence(UMovieSceneSequence* Sequence, EMovieSceneTimeUnit TimeUnit = EMovieSceneTimeUnit::DisplayRate);
+	static UE_API TArray<FMovieSceneMarkedFrame> GetMarkedFramesFromSequence(UMovieSceneSequence* Sequence, EMovieSceneTimeUnit TimeUnit = EMovieSceneTimeUnit::DisplayRate);
 
 	UE_DEPRECATED(5.4, "GetMarkedFrames is deprecated. Please use GetMarkedFrames that takes a time unit and defaults to display rate instead")
 	UFUNCTION(BlueprintCallable, Category = "Sequencer|Sequence", meta = (ScriptMethod, DeprecatedFunction, DeprecationMessage = "GetMarkedFrames is deprecated. Please use GetMarkedFrames that takes a time unit instead"))
@@ -541,7 +543,7 @@ public:
 	 * @return The index to the newly added marked frame
 	 */
 	UFUNCTION(BlueprintCallable, Category = "Sequencer|Sequence", meta = (ScriptMethod, DisplayName = "Add Marked Frame"))
-	static int32 AddMarkedFrameToSequence(UMovieSceneSequence* Sequence, const FMovieSceneMarkedFrame& InMarkedFrame, EMovieSceneTimeUnit TimeUnit = EMovieSceneTimeUnit::DisplayRate);
+	static UE_API int32 AddMarkedFrameToSequence(UMovieSceneSequence* Sequence, const FMovieSceneMarkedFrame& InMarkedFrame, EMovieSceneTimeUnit TimeUnit = EMovieSceneTimeUnit::DisplayRate);
 
 	UE_DEPRECATED(5.4, "AddMarkedFrame is deprecated. Please use AddMarkedFrame that takes a time unit and defaults to display rate instead")
 	UFUNCTION(BlueprintCallable, Category = "Sequencer|Sequence", meta = (ScriptMethod, DeprecatedFunction, DeprecationMessage = "AddMarkedFrame is deprecated. Please use AddMarkedFrame that takes a time unit instead"))
@@ -554,7 +556,7 @@ public:
 	 * @InFrameNumber The frame number to set
 	 */
 	UFUNCTION(BlueprintCallable, Category = "Sequencer|Sequence", meta = (ScriptMethod, DisplayName = "Set Marked Frame"))
-	static void SetMarkedFrameInSequence(UMovieSceneSequence* Sequence, int32 InMarkIndex, FFrameNumber InFrameNumber, EMovieSceneTimeUnit TimeUnit = EMovieSceneTimeUnit::DisplayRate);
+	static UE_API void SetMarkedFrameInSequence(UMovieSceneSequence* Sequence, int32 InMarkIndex, FFrameNumber InFrameNumber, EMovieSceneTimeUnit TimeUnit = EMovieSceneTimeUnit::DisplayRate);
 
 	UE_DEPRECATED(5.4, "SetMarkedFrame is deprecated. Please use SetMarkedFrame that takes a time unit and defaults to display rate instead")
 	UFUNCTION(BlueprintCallable, Category = "Sequencer|Sequence", meta = (ScriptMethod, DeprecatedFunction, DeprecationMessage = "SetMarkedFrame is deprecated. Please use SetMarkedFrame that takes a time unit instead"))
@@ -566,19 +568,19 @@ public:
 	 * @DeleteIndex The index to the user marked frame to delete
 	 */
 	UFUNCTION(BlueprintCallable, Category = "Sequencer|Sequence", meta = (ScriptMethod))
-	static void DeleteMarkedFrame(UMovieSceneSequence* Sequence, int32 DeleteIndex);
+	static UE_API void DeleteMarkedFrame(UMovieSceneSequence* Sequence, int32 DeleteIndex);
 
 	/*
 	 * Delete all user marked frames
 	 */
 	UFUNCTION(BlueprintCallable, Category = "Sequencer|Sequence", meta = (ScriptMethod))
-	static void DeleteMarkedFrames(UMovieSceneSequence* Sequence);
+	static UE_API void DeleteMarkedFrames(UMovieSceneSequence* Sequence);
 
 	/*
 	 * Sort the marked frames in chronological order
 	 */
 	UFUNCTION(BlueprintCallable, Category = "Sequencer|Sequence", meta = (ScriptMethod))
-	static void SortMarkedFrames(UMovieSceneSequence* Sequence);
+	static UE_API void SortMarkedFrames(UMovieSceneSequence* Sequence);
 
 	/*
 	 * Find the user marked frame by label
@@ -586,7 +588,7 @@ public:
 	 * @InLabel The label to the user marked frame to find
 	 */
 	UFUNCTION(BlueprintCallable, Category = "Sequencer|Sequence", meta = (ScriptMethod))
-	static int32 FindMarkedFrameByLabel(UMovieSceneSequence* Sequence, const FString& InLabel);
+	static UE_API int32 FindMarkedFrameByLabel(UMovieSceneSequence* Sequence, const FString& InLabel);
 
 	/*
 	 * Find the user marked frame by frame number
@@ -594,7 +596,7 @@ public:
 	 * @InFrameNumber The frame number of the user marked frame to find
 	 */
 	UFUNCTION(BlueprintCallable, Category = "Sequencer|Sequence", meta = (ScriptMethod, DisplayName = "Find Marked Frame By Frame Number"))
-	static int32 FindMarkedFrameByFrameNumberInSequence(UMovieSceneSequence* Sequence, FFrameNumber InFrameNumber, EMovieSceneTimeUnit TimeUnit = EMovieSceneTimeUnit::DisplayRate);
+	static UE_API int32 FindMarkedFrameByFrameNumberInSequence(UMovieSceneSequence* Sequence, FFrameNumber InFrameNumber, EMovieSceneTimeUnit TimeUnit = EMovieSceneTimeUnit::DisplayRate);
 
 	UE_DEPRECATED(5.4, "FindMarkedFrameByFrameNumber is deprecated. Please use FindMarkedFrameByFrameNumber that takes a time unit and defaults to display rate instead")
 	UFUNCTION(BlueprintCallable, Category = "Sequencer|Sequence", meta = (ScriptMethod, DeprecatedFunction, DeprecationMessage = "FindMarkedFrameByFrameNumber is deprecated. Please use FindMarkedFrameByFrameNumber that takes a time unit instead"))
@@ -607,7 +609,7 @@ public:
 	 * @bForward Find forward from the given frame number.
 	 */
 	UFUNCTION(BlueprintCallable, Category = "Sequencer|Sequence", meta = (ScriptMethod, DisplayName = "Find Next Marked Frame"))
-	static int32 FindNextMarkedFrameInSequence(UMovieSceneSequence* Sequence, FFrameNumber InFrameNumber, bool bForward, EMovieSceneTimeUnit TimeUnit = EMovieSceneTimeUnit::DisplayRate);
+	static UE_API int32 FindNextMarkedFrameInSequence(UMovieSceneSequence* Sequence, FFrameNumber InFrameNumber, bool bForward, EMovieSceneTimeUnit TimeUnit = EMovieSceneTimeUnit::DisplayRate);
 
 	UE_DEPRECATED(5.4, "FindNextMarkedFrame is deprecated. Please use FindNextMarkedFrame that takes a time unit instead")
 	UFUNCTION(BlueprintCallable, Category = "Sequencer|Sequence", meta = (ScriptMethod, DeprecatedFunction, DeprecationMessage = "FindNextMarkedFrame is deprecated. Please use FindNextMarkedFrame that takes a time unit and defaults to display rate instead"))
@@ -619,7 +621,7 @@ public:
 	 * @bInReadOnly Whether the movie scene should be read only or not
 	 */
 	UFUNCTION(BlueprintCallable, Category = "Sequencer|Sequence", meta = (ScriptMethod))
-	static void SetReadOnly(UMovieSceneSequence* Sequence, bool bInReadOnly);
+	static UE_API void SetReadOnly(UMovieSceneSequence* Sequence, bool bInReadOnly);
 
 	/*
 	 * Is read only
@@ -627,7 +629,7 @@ public:
 	 * @return Whether the movie scene is read only or not
 	 */
 	UFUNCTION(BlueprintCallable, Category = "Sequencer|Sequence", meta = (ScriptMethod))
-	static bool IsReadOnly(UMovieSceneSequence* Sequence);
+	static UE_API bool IsReadOnly(UMovieSceneSequence* Sequence);
 
 	/*
 	 * Set playback range locked
@@ -635,7 +637,7 @@ public:
 	 * @bInLocked Whether the movie scene playback range should be locked
 	 */
 	UFUNCTION(BlueprintCallable, Category = "Sequencer|Sequence", meta = (ScriptMethod))
-	static void SetPlaybackRangeLocked(UMovieSceneSequence* Sequence, bool bInLocked);
+	static UE_API void SetPlaybackRangeLocked(UMovieSceneSequence* Sequence, bool bInLocked);
 
 	/*
 	 * Is playback ranged locked
@@ -643,7 +645,7 @@ public:
 	 * @return Whether the movie scene playback range is locked
 	 */
 	UFUNCTION(BlueprintCallable, Category = "Sequencer|Sequence", meta = (ScriptMethod))
-	static bool IsPlaybackRangeLocked(UMovieSceneSequence* Sequence);
+	static UE_API bool IsPlaybackRangeLocked(UMovieSceneSequence* Sequence);
 
 	/*
 	 * Set marked frames locked
@@ -651,7 +653,7 @@ public:
 	 * @bInLocked Whether the movie scene marked frames should be locked
 	 */
 	UFUNCTION(BlueprintCallable, Category = "Sequencer|Sequence", meta = (ScriptMethod))
-	static void SetMarkedFramesLocked(UMovieSceneSequence* Sequence, bool bInLocked);
+	static UE_API void SetMarkedFramesLocked(UMovieSceneSequence* Sequence, bool bInLocked);
 
 	/*
 	 * Are marked frames locked
@@ -659,5 +661,7 @@ public:
 	 * @return Whether the movie scene marked frames are locked
 	 */
 	UFUNCTION(BlueprintCallable, Category = "Sequencer|Sequence", meta = (ScriptMethod))
-	static bool AreMarkedFramesLocked(UMovieSceneSequence* Sequence);
+	static UE_API bool AreMarkedFramesLocked(UMovieSceneSequence* Sequence);
 };
+
+#undef UE_API
diff --git a/Engine/Plugins/MovieScene/SequencerScripting/Source/SequencerScripting/Public/ExtensionLibraries/MovieSceneTimeWarpExtensions.h b/Engine/Plugins/MovieScene/SequencerScripting/Source/SequencerScripting/Public/ExtensionLibraries/MovieSceneTimeWarpExtensions.h
--- a/Engine/Plugins/MovieScene/SequencerScripting/Source/SequencerScripting/Public/ExtensionLibraries/MovieSceneTimeWarpExtensions.h
+++ b/Engine/Plugins/MovieScene/SequencerScripting/Source/SequencerScripting/Public/ExtensionLibraries/MovieSceneTimeWarpExtensions.h
@@ -7,6 +7,7 @@
 
 struct FMovieSceneTimeWarpVariant;
 
+#define UE_API SEQUENCERSCRIPTING_API
 
 /**
  * Function library containing methods that relate to time-warp within Sequencer
@@ -22,31 +23,33 @@ public:
 	 * Converts a timewarp variant struct to a constant play rate
 	 */
 	UFUNCTION(BlueprintPure, meta=(DisplayName = "To PlayRate", BlueprintAutocast), Category="Sequencer|TimeWarp")
-	static SEQUENCERSCRIPTING_API double Conv_TimeWarpVariantToPlayRate(const FMovieSceneTimeWarpVariant& TimeWarp);
+	static UE_API double Conv_TimeWarpVariantToPlayRate(const FMovieSceneTimeWarpVariant& TimeWarp);
 
 	/**
 	 * Converts a constant playrate to a timewarp variant
 	 */
 	UFUNCTION(BlueprintPure, meta=(DisplayName = "To TimeWarp", BlueprintAutocast), Category="Sequencer|TimeWarp")
-	static SEQUENCERSCRIPTING_API FMovieSceneTimeWarpVariant Conv_PlayRateToTimeWarpVariant(double ConstantPlayRate);
+	static UE_API FMovieSceneTimeWarpVariant Conv_PlayRateToTimeWarpVariant(double ConstantPlayRate);
 
 	/**
 	 * Retrieve this timewarp's constant play rate. Will throw an error if the timewarp is not a constant play rate.
 	 */
 	UFUNCTION(BlueprintPure, meta=(ScriptMethod), Category="Sequencer|TimeWarp")
-	static SEQUENCERSCRIPTING_API double ToFixedPlayRate(const FMovieSceneTimeWarpVariant& TimeWarp);
+	static UE_API double ToFixedPlayRate(const FMovieSceneTimeWarpVariant& TimeWarp);
 
 	/**
 	 * Assign a constant playrate to this timewarp, overwriting any existing timewarp implementation.
 	 */
 	UFUNCTION(BlueprintPure, meta=(ScriptMethod), Category="Sequencer|TimeWarp")
-	static SEQUENCERSCRIPTING_API void SetFixedPlayRate(UPARAM(ref) FMovieSceneTimeWarpVariant& TimeWarp, double FixedPlayRate);
+	static UE_API void SetFixedPlayRate(UPARAM(ref) FMovieSceneTimeWarpVariant& TimeWarp, double FixedPlayRate);
 
 public:
 
 	UFUNCTION(BlueprintPure, Category="Sequencer|TimeWarp")
-	static SEQUENCERSCRIPTING_API void BreakTimeWarp(const FMovieSceneTimeWarpVariant& TimeWarp, double& FixedPlayRate);
+	static UE_API void BreakTimeWarp(const FMovieSceneTimeWarpVariant& TimeWarp, double& FixedPlayRate);
 
 	UFUNCTION(BlueprintPure, Category="Sequencer|TimeWarp")
-	static SEQUENCERSCRIPTING_API FMovieSceneTimeWarpVariant MakeTimeWarp(double FixedPlayRate);
-};
\ No newline at end of file
+	static UE_API FMovieSceneTimeWarpVariant MakeTimeWarp(double FixedPlayRate);
+};
+
+#undef UE_API
diff --git a/Engine/Plugins/MovieScene/SequencerScripting/Source/SequencerScripting/Public/ExtensionLibraries/MovieSceneTrackExtensions.h b/Engine/Plugins/MovieScene/SequencerScripting/Source/SequencerScripting/Public/ExtensionLibraries/MovieSceneTrackExtensions.h
--- a/Engine/Plugins/MovieScene/SequencerScripting/Source/SequencerScripting/Public/ExtensionLibraries/MovieSceneTrackExtensions.h
+++ b/Engine/Plugins/MovieScene/SequencerScripting/Source/SequencerScripting/Public/ExtensionLibraries/MovieSceneTrackExtensions.h
@@ -11,6 +11,8 @@ class FText;
 class UMovieSceneTrack;
 class UMovieSceneSection;
 
+#define UE_API SEQUENCERSCRIPTING_API
+
 /**
  * Function library containing methods that should be hoisted onto UMovieSceneTracks for scripting
  */
@@ -28,7 +30,7 @@ public:
 	 * @param InName The name for this track
 	 */
 	UFUNCTION(BlueprintCallable, Category = "Sequencer|Track", meta=(ScriptMethod, DevelopmentOnly))
-	static void SetDisplayName(UMovieSceneTrack* Track, const FText& InName);
+	static UE_API void SetDisplayName(UMovieSceneTrack* Track, const FText& InName);
 
 	/**
 	 * Get this track's display name
@@ -37,7 +39,7 @@ public:
 	 * @return This track's display name
 	 */
 	UFUNCTION(BlueprintCallable, Category = "Sequencer|Track", meta=(ScriptMethod, DevelopmentOnly))
-	static FText GetDisplayName(UMovieSceneTrack* Track);
+	static UE_API FText GetDisplayName(UMovieSceneTrack* Track);
 
 	/**
 	 * Set this track row's display name
@@ -47,7 +49,7 @@ public:
 	 * @param RowIndex The row index for the track
 	 */
 	UFUNCTION(BlueprintCallable, Category = "Sequencer|Track", meta=(ScriptMethod, DevelopmentOnly))
-	static void SetTrackRowDisplayName(UMovieSceneTrack* Track, const FText& InName, int32 RowIndex);
+	static UE_API void SetTrackRowDisplayName(UMovieSceneTrack* Track, const FText& InName, int32 RowIndex);
 
 	/**
 	 * Get this track row's display name
@@ -58,7 +60,7 @@ public:
 	 * 
 	 */
 	UFUNCTION(BlueprintCallable, Category = "Sequencer|Track", meta=(ScriptMethod, DevelopmentOnly))
-	static FText GetTrackRowDisplayName(UMovieSceneTrack* Track, int32 RowIndex);
+	static UE_API FText GetTrackRowDisplayName(UMovieSceneTrack* Track, int32 RowIndex);
 
 	/**
 	 * Add a new section to this track
@@ -67,7 +69,7 @@ public:
 	 * @return The newly create section if successful
 	 */
 	UFUNCTION(BlueprintCallable, Category = "Sequencer|Track", meta=(ScriptMethod))
-	static UMovieSceneSection* AddSection(UMovieSceneTrack* Track);
+	static UE_API UMovieSceneSection* AddSection(UMovieSceneTrack* Track);
 
 	/**
 	 * Access all this track's sections
@@ -76,7 +78,7 @@ public:
 	 * @return An array of this track's sections
 	 */
 	UFUNCTION(BlueprintCallable, Category = "Sequencer|Track", meta=(ScriptMethod))
-	static TArray<UMovieSceneSection*> GetSections(UMovieSceneTrack* Track);
+	static UE_API TArray<UMovieSceneSection*> GetSections(UMovieSceneTrack* Track);
 
 	/**
 	 * Remove the specified section
@@ -85,7 +87,7 @@ public:
 	 * @param Section      The section to remove
 	 */
 	UFUNCTION(BlueprintCallable, Category = "Sequencer|Track", meta=(ScriptMethod))
-	static void RemoveSection(UMovieSceneTrack* Track, UMovieSceneSection* Section);
+	static UE_API void RemoveSection(UMovieSceneTrack* Track, UMovieSceneSection* Section);
 
 	/**
 	 * Get the sorting order for this track
@@ -94,7 +96,7 @@ public:
 	 * @return The sorting order of the requested track
 	 */
 	UFUNCTION(BlueprintCallable, Category = "Sequencer|Track", meta=(ScriptMethod, DevelopmentOnly))
-	static int32 GetSortingOrder(UMovieSceneTrack* Track);
+	static UE_API int32 GetSortingOrder(UMovieSceneTrack* Track);
  
 	/**
 	 * Set the sorting order for this track
@@ -103,7 +105,7 @@ public:
 	 * @param SortingOrder The sorting order to set
 	 */
 	UFUNCTION(BlueprintCallable, Category = "Sequencer|Track", meta=(ScriptMethod, DevelopmentOnly))
-	static void SetSortingOrder(UMovieSceneTrack* Track, int32 SortingOrder);
+	static UE_API void SetSortingOrder(UMovieSceneTrack* Track, int32 SortingOrder);
  
 	/**
 	 * Get the color tint for this track
@@ -112,7 +114,7 @@ public:
 	 * @return The color tint of the requested track
 	 */
 	UFUNCTION(BlueprintCallable, Category = "Sequencer|Track", meta=(ScriptMethod, DevelopmentOnly))
-	static FColor GetColorTint(UMovieSceneTrack* Track);
+	static UE_API FColor GetColorTint(UMovieSceneTrack* Track);
  
 	/**
 	 * Set the color tint for this track
@@ -121,7 +123,7 @@ public:
 	 * @param ColorTint The color tint to set
 	 */
 	UFUNCTION(BlueprintCallable, Category = "Sequencer|Track", meta=(ScriptMethod, DevelopmentOnly))
-	static void SetColorTint(UMovieSceneTrack* Track, const FColor& ColorTint);
+	static UE_API void SetColorTint(UMovieSceneTrack* Track, const FColor& ColorTint);
 
 	/**
 	 * Get the section to key for this track
@@ -130,7 +132,7 @@ public:
 	 * @return The section to key for the requested track
 	 */
 	UFUNCTION(BlueprintCallable, Category = "Sequencer|Track", meta=(ScriptMethod, DevelopmentOnly))
-	static UMovieSceneSection* GetSectionToKey(UMovieSceneTrack* Track);
+	static UE_API UMovieSceneSection* GetSectionToKey(UMovieSceneTrack* Track);
  
 	/**
 	 * Set the section to key for this track. When properties for this section are modified externally, 
@@ -141,5 +143,7 @@ public:
 	 * @param Section      The section to key for this track
 	 */
 	UFUNCTION(BlueprintCallable, Category = "Sequencer|Track", meta=(ScriptMethod, DevelopmentOnly))
-	static void SetSectionToKey(UMovieSceneTrack* Track, UMovieSceneSection* Section); 
+	static UE_API void SetSectionToKey(UMovieSceneTrack* Track, UMovieSceneSection* Section);
 };
+
+#undef UE_API
diff --git a/Engine/Plugins/MovieScene/SequencerScripting/Source/SequencerScripting/Public/ExtensionLibraries/MovieSceneVectorTrackExtensions.h b/Engine/Plugins/MovieScene/SequencerScripting/Source/SequencerScripting/Public/ExtensionLibraries/MovieSceneVectorTrackExtensions.h
--- a/Engine/Plugins/MovieScene/SequencerScripting/Source/SequencerScripting/Public/ExtensionLibraries/MovieSceneVectorTrackExtensions.h
+++ b/Engine/Plugins/MovieScene/SequencerScripting/Source/SequencerScripting/Public/ExtensionLibraries/MovieSceneVectorTrackExtensions.h
@@ -9,6 +9,8 @@
 class UMovieSceneDoubleVectorTrack;
 class UMovieSceneFloatVectorTrack;
 
+#define UE_API SEQUENCERSCRIPTING_API
+
 /**
  * Function library containing methods that should be hoisted onto UMovieSceneFloatVectorTrack for scripting
  */
@@ -26,7 +28,7 @@ public:
 	 * @param InNumChannelsUsed The number of channels to use for this track
 	 */
 	UFUNCTION(BlueprintCallable, Category = "Sequencer|Track", meta = (ScriptMethod))
-	static void SetNumChannelsUsed(UMovieSceneFloatVectorTrack* Track, int32 InNumChannelsUsed);
+	static UE_API void SetNumChannelsUsed(UMovieSceneFloatVectorTrack* Track, int32 InNumChannelsUsed);
 
 	/**
 	 * Get the number of channels used for this track
@@ -35,7 +37,7 @@ public:
 	 * @return The number of channels used for this track
 	 */
 	UFUNCTION(BlueprintCallable, Category = "Sequencer|Track", meta = (ScriptMethod))
-	static int32 GetNumChannelsUsed(UMovieSceneFloatVectorTrack* Track);
+	static UE_API int32 GetNumChannelsUsed(UMovieSceneFloatVectorTrack* Track);
 };
 
 /**
@@ -55,7 +57,7 @@ public:
 	 * @param InNumChannelsUsed The number of channels to use for this track
 	 */
 	UFUNCTION(BlueprintCallable, Category = "Sequencer|Track", meta = (ScriptMethod))
-	static void SetNumChannelsUsed(UMovieSceneDoubleVectorTrack* Track, int32 InNumChannelsUsed);
+	static UE_API void SetNumChannelsUsed(UMovieSceneDoubleVectorTrack* Track, int32 InNumChannelsUsed);
 
 	/**
 	 * Get the number of channels used for this track
@@ -64,6 +66,8 @@ public:
 	 * @return The number of channels used for this track
 	 */
 	UFUNCTION(BlueprintCallable, Category = "Sequencer|Track", meta = (ScriptMethod))
-	static int32 GetNumChannelsUsed(UMovieSceneDoubleVectorTrack* Track);
+	static UE_API int32 GetNumChannelsUsed(UMovieSceneDoubleVectorTrack* Track);
 };
 
+#undef UE_API
+
diff --git a/Engine/Plugins/MovieScene/SequencerScripting/Source/SequencerScripting/Public/ExtensionLibraries/SequencerScriptingRangeExtensions.h b/Engine/Plugins/MovieScene/SequencerScripting/Source/SequencerScripting/Public/ExtensionLibraries/SequencerScriptingRangeExtensions.h
--- a/Engine/Plugins/MovieScene/SequencerScripting/Source/SequencerScripting/Public/ExtensionLibraries/SequencerScriptingRangeExtensions.h
+++ b/Engine/Plugins/MovieScene/SequencerScripting/Source/SequencerScripting/Public/ExtensionLibraries/SequencerScriptingRangeExtensions.h
@@ -8,6 +8,8 @@
 
 struct FSequencerScriptingRange;
 
+#define UE_API SEQUENCERSCRIPTING_API
+
 /**
  * Function library containing methods that should be hoisted onto FSequencerScriptingRanges
  */
@@ -22,7 +24,7 @@ public:
 	 * @param Range       The range to check
 	 */
 	UFUNCTION(BlueprintCallable, Category = "Sequencer|Sequence", meta=(ScriptMethod))
-	static bool HasStart(const FSequencerScriptingRange& Range);
+	static UE_API bool HasStart(const FSequencerScriptingRange& Range);
 
 	/**
 	 * Check whether this range has an end
@@ -30,7 +32,7 @@ public:
 	 * @param Range       The range to check
 	 */
 	UFUNCTION(BlueprintCallable, Category = "Sequencer|Sequence", meta=(ScriptMethod))
-	static bool HasEnd(const FSequencerScriptingRange& Range);
+	static UE_API bool HasEnd(const FSequencerScriptingRange& Range);
 
 	/**
 	 * Remove the start from this range, making it infinite
@@ -38,7 +40,7 @@ public:
 	 * @param Range       The range to remove the start from
 	 */
 	UFUNCTION(BlueprintCallable, Category = "Sequencer|Sequence", meta=(ScriptMethod))
-	static void RemoveStart(UPARAM(ref) FSequencerScriptingRange& Range);
+	static UE_API void RemoveStart(UPARAM(ref) FSequencerScriptingRange& Range);
 
 	/**
 	 * Remove the end from this range, making it infinite
@@ -46,7 +48,7 @@ public:
 	 * @param Range       The range to remove the end from
 	 */
 	UFUNCTION(BlueprintCallable, Category = "Sequencer|Sequence", meta=(ScriptMethod))
-	static void RemoveEnd(UPARAM(ref) FSequencerScriptingRange& Range);
+	static UE_API void RemoveEnd(UPARAM(ref) FSequencerScriptingRange& Range);
 
 	/**
 	 * Get the starting time for the specified range in seconds, if it has one. Defined as the first valid time that is inside the range.
@@ -54,7 +56,7 @@ public:
 	 * @param Range       The range to get the start from
 	 */
 	UFUNCTION(BlueprintCallable, Category = "Sequencer|Sequence", meta=(ScriptMethod))
-	static float GetStartSeconds(const FSequencerScriptingRange& Range);
+	static UE_API float GetStartSeconds(const FSequencerScriptingRange& Range);
 
 	/**
 	 * Get the ending time for the specified range in seconds, if it has one. Defined as the first time that is outside of the range.
@@ -62,7 +64,7 @@ public:
 	 * @param Range       The range to get the end from
 	 */
 	UFUNCTION(BlueprintCallable, Category = "Sequencer|Sequence", meta=(ScriptMethod))
-	static float GetEndSeconds(const FSequencerScriptingRange& Range);
+	static UE_API float GetEndSeconds(const FSequencerScriptingRange& Range);
 
 	/**
 	 * Set the starting time for the specified range in seconds. Interpreted as the first valid time that is inside the range.
@@ -70,7 +72,7 @@ public:
 	 * @param Range       The range to set the start on
 	 */
 	UFUNCTION(BlueprintCallable, Category = "Sequencer|Sequence", meta=(ScriptMethod))
-	static void SetStartSeconds(UPARAM(ref) FSequencerScriptingRange& Range, float Start);
+	static UE_API void SetStartSeconds(UPARAM(ref) FSequencerScriptingRange& Range, float Start);
 
 	/**
 	 * Set the ending time for the specified range in seconds. Interpreted as the first time that is outside of the range.
@@ -78,7 +80,7 @@ public:
 	 * @param Range       The range to set the end on
 	 */
 	UFUNCTION(BlueprintCallable, Category = "Sequencer|Sequence", meta=(ScriptMethod))
-	static void SetEndSeconds(UPARAM(ref) FSequencerScriptingRange& Range, float End);
+	static UE_API void SetEndSeconds(UPARAM(ref) FSequencerScriptingRange& Range, float End);
 
 	/**
 	 * Get the starting frame for the specified range, if it has one. Defined as the first valid frame that is inside the range.
@@ -86,7 +88,7 @@ public:
 	 * @param Range       The range to get the start from
 	 */
 	UFUNCTION(BlueprintCallable, Category = "Sequencer|Sequence", meta=(ScriptMethod))
-	static int32 GetStartFrame(const FSequencerScriptingRange& Range);
+	static UE_API int32 GetStartFrame(const FSequencerScriptingRange& Range);
 
 	/**
 	 * Get the ending frame for the specified range, if it has one. Defined as the first subsequent frame that is outside of the range.
@@ -94,7 +96,7 @@ public:
 	 * @param Range       The range to get the end from
 	 */
 	UFUNCTION(BlueprintCallable, Category = "Sequencer|Sequence", meta=(ScriptMethod))
-	static int32 GetEndFrame(const FSequencerScriptingRange& Range);
+	static UE_API int32 GetEndFrame(const FSequencerScriptingRange& Range);
 
 	/**
 	 * Set the starting frame for the specified range. Interpreted as the first valid frame that is inside the range.
@@ -102,7 +104,7 @@ public:
 	 * @param Range       The range to set the start on
 	 */
 	UFUNCTION(BlueprintCallable, Category = "Sequencer|Sequence", meta=(ScriptMethod))
-	static void SetStartFrame(UPARAM(ref) FSequencerScriptingRange& Range, int32 Start);
+	static UE_API void SetStartFrame(UPARAM(ref) FSequencerScriptingRange& Range, int32 Start);
 
 	/**
 	 * Set the ending frame for the specified range. Interpreted as the first subsequent frame that is outside of the range.
@@ -110,5 +112,7 @@ public:
 	 * @param Range       The range to set the end on
 	 */
 	UFUNCTION(BlueprintCallable, Category = "Sequencer|Sequence", meta=(ScriptMethod))
-	static void SetEndFrame(UPARAM(ref) FSequencerScriptingRange& Range, int32 End);
+	static UE_API void SetEndFrame(UPARAM(ref) FSequencerScriptingRange& Range, int32 End);
 };
+
+#undef UE_API
