Fix unresolved external errors for merged module builds

#rnx

diff --git a/Engine/Plugins/MovieScene/SequencerScripting/Source/SequencerScriptingEditor/Private/SequencerTools.cpp b/Engine/Plugins/MovieScene/SequencerScripting/Source/SequencerScriptingEditor/Private/SequencerTools.cpp
--- a/Engine/Plugins/MovieScene/SequencerScripting/Source/SequencerScriptingEditor/Private/SequencerTools.cpp
+++ b/Engine/Plugins/MovieScene/SequencerScripting/Source/SequencerScriptingEditor/Private/SequencerTools.cpp
@@ -41,7 +41,7 @@
 
 #define LOCTEXT_NAMESPACE "SequencerTools"
 
-DECLARE_LOG_CATEGORY_CLASS(LogSequencer, Log, All);
+DECLARE_LOG_CATEGORY_CLASS(LogSequencerTools, Log, All);
 
 namespace SequencerToolsPrivate
 {
@@ -52,7 +52,7 @@ namespace SequencerToolsPrivate
 	{
 		if (!Condition.IsBound() || !Callback.IsSet() || !World)
 		{
-			UE_LOG(LogSequencer, Error, TEXT("WaitUntil() - One of the parameters is invalid."));
+			UE_LOG(LogSequencerTools, Error, TEXT("WaitUntil() - One of the parameters is invalid."));
 			return;
 		}
 
@@ -107,7 +107,7 @@ PRAGMA_ENABLE_DEPRECATION_WARNINGS
 	{
 		if (!LevelSequenceCapture->LevelSequenceAsset.IsValid())
 		{
-			// UE_LOG(LogSequencer, Warning, TEXT("No Level Sequence Asset specified in UAutomatedLevelSequenceCapture."));
+			// UE_LOG(LogSequencerTools, Warning, TEXT("No Level Sequence Asset specified in UAutomatedLevelSequenceCapture."));
 			FFrame::KismetExecutionMessage(TEXT("No Level Sequence Asset specified in UAutomatedLevelSequenceCapture."), ELogVerbosity::Error);
 			return false;
 		}
@@ -433,7 +433,7 @@ bool USequencerToolsFunctionLibrary::ExportAnimSequenceWaitForDelegate(UWorld* W
 			}
 			else
 			{
-				UE_LOG(LogSequencer, Warning, TEXT("USequencerToolsFunctionLibrary::ExportAnimSequence: No skel mesh found in Sequencer"));
+				UE_LOG(LogSequencerTools, Warning, TEXT("USequencerToolsFunctionLibrary::ExportAnimSequence: No skel mesh found in Sequencer"));
 				bResult = false;
 			}
 
@@ -442,7 +442,7 @@ bool USequencerToolsFunctionLibrary::ExportAnimSequenceWaitForDelegate(UWorld* W
 		}
 		else
 		{
-			UE_LOG(LogSequencer, Warning, TEXT("USequencerToolsFunctionLibrary::ExportAnimSequence: One or more of the completion lambda's captured parameters are invalid."));
+			UE_LOG(LogSequencerTools, Warning, TEXT("USequencerToolsFunctionLibrary::ExportAnimSequence: One or more of the completion lambda's captured parameters are invalid."));
 			bResult = false;
 		}
 
