Anim In Engine: Add Scripting Channels for spaces and constraints. Had to expose the base scripting channels/keys from the SequencerScripting, and added ways to add/remove TFunctions to MovieSceneExtension to get custom channel types. Also exposed a function to get/add a space channel if one doesn't exist.

#Synced-CL 48371439

diff --git a/Engine/Plugins/Animation/ControlRig/Source/ControlRig/Public/Sequencer/MovieSceneControlRigSpaceChannel.h b/Engine/Plugins/Animation/ControlRig/Source/ControlRig/Public/Sequencer/MovieSceneControlRigSpaceChannel.h
--- a/Engine/Plugins/Animation/ControlRig/Source/ControlRig/Public/Sequencer/MovieSceneControlRigSpaceChannel.h
+++ b/Engine/Plugins/Animation/ControlRig/Source/ControlRig/Public/Sequencer/MovieSceneControlRigSpaceChannel.h
@@ -31,7 +31,7 @@ enum class EMovieSceneControlRigSpaceType : uint8
 	ControlRig
 };
 
-USTRUCT()
+USTRUCT(BlueprintType)
 struct FMovieSceneControlRigSpaceBaseKey
 {
 	GENERATED_BODY()
@@ -48,9 +48,9 @@ struct FMovieSceneControlRigSpaceBaseKey
 
 	UE_API FName GetName() const;
 
-	UPROPERTY(EditAnywhere, Category = "Key")
+	UPROPERTY(BlueprintReadWrite,EditAnywhere, Category = "Key")
 	EMovieSceneControlRigSpaceType SpaceType = EMovieSceneControlRigSpaceType::Parent;
-	UPROPERTY(EditAnywhere, Category = "Key")
+	UPROPERTY(BlueprintReadWrite,EditAnywhere, Category = "Key")
 	FRigElementKey ControlRigElement;
 
 };
diff --git a/Engine/Plugins/Animation/ControlRig/Source/ControlRigEditor/Private/ControlRigEditorModule.cpp b/Engine/Plugins/Animation/ControlRig/Source/ControlRigEditor/Private/ControlRigEditorModule.cpp
--- a/Engine/Plugins/Animation/ControlRig/Source/ControlRigEditor/Private/ControlRigEditorModule.cpp
+++ b/Engine/Plugins/Animation/ControlRig/Source/ControlRigEditor/Private/ControlRigEditorModule.cpp
@@ -42,6 +42,7 @@
 #include "ControlRigHierarchyCommands.h"
 #include "ControlRigModularRigCommands.h"
 #include "Editor/ModularRigEventQueueCommands.h"
+#include "ExtensionLibraries/MovieSceneSectionExtensions.h"
 #include "RigDependencyGraph/RigDependencyGraphCommands.h"
 #include "Animation/AnimSequence.h"
 #include "Editor/ControlRigEditorEditMode.h"
@@ -102,6 +103,8 @@
 #include "Async/TaskGraphInterfaces.h"
 #include "AssetRegistry/IAssetRegistry.h"
 #include "SequencerUtilities.h"
+#include "Sequencer/Scripting/MovieSceneScriptingSpace.h"
+#include "Sequencer/Scripting/MovieSceneScriptingConstraint.h"
 #include "Bindings/MovieSceneSpawnableActorBinding.h"
 #if WITH_RIGVMLEGACYEDITOR
 #include "Editor/ControlRigLegacyEditor.h"
@@ -119,6 +122,7 @@
 
 DEFINE_LOG_CATEGORY(LogControlRigEditor);
 
+
 void FControlRigEditorModule::StartupModule()
 {
 	using namespace UE::ControlRigEditor;
@@ -190,19 +194,19 @@ void FControlRigEditorModule::StartupModule()
 
 	PropertiesToUnregisterOnShutdown.Add(FRigUnit_HierarchyAddControl_ShapeSettings::StaticStruct()->GetFName());
 	PropertyEditorModule.RegisterCustomPropertyTypeLayout(PropertiesToUnregisterOnShutdown.Last(), FOnGetPropertyTypeCustomizationInstance::CreateStatic(&FControlRigShapeSettingsDetails::MakeInstance));
-	
+
 
 	// Register asset tools
 	auto RegisterAssetTypeAction = [this](const TSharedRef<IAssetTypeActions>& InAssetTypeAction)
-	{
-		IAssetTools& AssetTools = FModuleManager::LoadModuleChecked<FAssetToolsModule>("AssetTools").Get();
-		RegisteredAssetTypeActions.Add(InAssetTypeAction);
-		AssetTools.RegisterAssetTypeActions(InAssetTypeAction);
-	};
+		{
+			IAssetTools& AssetTools = FModuleManager::LoadModuleChecked<FAssetToolsModule>("AssetTools").Get();
+			RegisteredAssetTypeActions.Add(InAssetTypeAction);
+			AssetTools.RegisterAssetTypeActions(InAssetTypeAction);
+		};
 
 	RegisterAssetTypeAction(MakeShareable(new FControlRigBlueprintActions()));
 	RegisterAssetTypeAction(MakeShareable(new FControlRigShapeLibraryActions()));
-	
+
 	// Register sequencer track editor
 	ISequencerModule& SequencerModule = FModuleManager::Get().LoadModuleChecked<ISequencerModule>("Sequencer");
 	SequencerModule.RegisterChannelInterface<FMovieSceneControlRigSpaceChannel>();
@@ -211,7 +215,7 @@ void FControlRigEditorModule::StartupModule()
 	// register UTransformableControlHandle animatable interface
 	FConstraintChannelInterfaceRegistry& ConstraintChannelInterfaceRegistry = FConstraintChannelInterfaceRegistry::Get();
 	ConstraintChannelInterfaceRegistry.RegisterConstraintChannelInterface<UTransformableControlHandle>(MakeUnique<FControlConstraintChannelInterface>());
-	
+
 	AddControlRigExtenderToToolMenu("AssetEditor.AnimationEditor.ToolBar");
 
 	FEditorModeRegistry::Get().RegisterMode<FControlRigEditMode>(
@@ -266,6 +270,35 @@ void FControlRigEditorModule::StartupModule()
 
 	RigDependencyGraphPanelNodeFactory = MakeShareable(new FRigDependencyGraphPanelNodeFactory());
 	FEdGraphUtilities::RegisterVisualNodeFactory(RigDependencyGraphPanelNodeFactory);
+
+	auto RegisterSpaceChannel = [](UMovieSceneSection* Section, TSubclassOf<UMovieSceneScriptingChannel> ChannelType, TArray<UMovieSceneScriptingChannel*>& OutChannels) -> bool
+		{
+			if (ChannelType == UMovieSceneScriptingChannel::StaticClass() ||
+				ChannelType == UMovieSceneScriptingSpaceChannel::StaticClass())
+			{
+				FMovieSceneChannelProxy& ChannelProxy = Section->GetChannelProxy();
+				TWeakObjectPtr<UMovieSceneSequence> Sequence = Section->GetTypedOuter<UMovieSceneSequence>();
+				GetScriptingChannelsForChannel<FMovieSceneControlRigSpaceChannel, UMovieSceneScriptingSpaceChannel>(ChannelProxy, Sequence, Section, OutChannels);
+				return true;
+			}
+			return false;
+		};
+	auto RegisterConstraintChannel = [](UMovieSceneSection* Section, TSubclassOf<UMovieSceneScriptingChannel> ChannelType, TArray<UMovieSceneScriptingChannel*>& OutChannels)-> bool
+		{
+			if (ChannelType == UMovieSceneScriptingChannel::StaticClass() ||
+				ChannelType == UMovieSceneScriptingConstraintChannel::StaticClass())
+			{
+				FMovieSceneChannelProxy& ChannelProxy = Section->GetChannelProxy();
+				TWeakObjectPtr<UMovieSceneSequence> Sequence = Section->GetTypedOuter<UMovieSceneSequence>();
+				GetScriptingChannelsForChannel<FMovieSceneConstraintChannel, UMovieSceneScriptingConstraintChannel>(ChannelProxy, Sequence, Section, OutChannels);
+				return true;
+			}
+			return false;
+		};
+
+	SpaceCustomChannelGuid = UMovieSceneSectionExtensions::SetGetCustomChannelsFunction(RegisterSpaceChannel);
+	ConstraintCustomChannelGuid = UMovieSceneSectionExtensions::SetGetCustomChannelsFunction(RegisterConstraintChannel);
+
 }
 
 void FControlRigEditorModule::ShutdownModule()
@@ -334,6 +367,8 @@ void FControlRigEditorModule::ShutdownModule()
 		FEdGraphUtilities::UnregisterVisualNodeFactory(RigDependencyGraphPanelNodeFactory);
 		RigDependencyGraphPanelNodeFactory.Reset();
 	}
+	UMovieSceneSectionExtensions::RemoveGetCustomChannelsFunction(SpaceCustomChannelGuid);
+	UMovieSceneSectionExtensions::RemoveGetCustomChannelsFunction(ConstraintCustomChannelGuid);
 }
 
 UClass* FControlRigEditorModule::GetRigVMAssetClass() const
diff --git a/Engine/Plugins/Animation/ControlRig/Source/ControlRigEditor/Private/ControlRigEditorModule.h b/Engine/Plugins/Animation/ControlRig/Source/ControlRigEditor/Private/ControlRigEditorModule.h
--- a/Engine/Plugins/Animation/ControlRig/Source/ControlRigEditor/Private/ControlRigEditorModule.h
+++ b/Engine/Plugins/Animation/ControlRig/Source/ControlRigEditor/Private/ControlRigEditorModule.h
@@ -80,4 +80,9 @@ private:
 	TArray<int32> WorkflowHandles;
 	
 	TSharedPtr<FRigDependencyGraphPanelNodeFactory> RigDependencyGraphPanelNodeFactory;
+
+	/** Guids for registered functions for creating custom UObjects for space and constraint channels for scripting*/
+	FGuid SpaceCustomChannelGuid;
+	FGuid ConstraintCustomChannelGuid;
+
 };
\ No newline at end of file
diff --git a/Engine/Plugins/Animation/ControlRig/Source/ControlRigEditor/Private/ControlRigSequencerEditorLibrary.cpp b/Engine/Plugins/Animation/ControlRig/Source/ControlRigEditor/Private/ControlRigSequencerEditorLibrary.cpp
--- a/Engine/Plugins/Animation/ControlRig/Source/ControlRigEditor/Private/ControlRigSequencerEditorLibrary.cpp
+++ b/Engine/Plugins/Animation/ControlRig/Source/ControlRigEditor/Private/ControlRigSequencerEditorLibrary.cpp
@@ -55,7 +55,9 @@
 #include "Tools/EvaluateSequencerTools.h"
 #include "Sequencer/AnimLayers/AnimLayers.h"
 #include "Sequencer/SelectionSets/SelectionSets.h"
-
+#include "Sequencer/Scripting/MovieSceneScriptingSpace.h"
+#include "Sequencer/Scripting/MovieSceneScriptingConstraint.h"
+#include "ExtensionLibraries/MovieSceneSectionExtensions.h"
 
 #include UE_INLINE_GENERATED_CPP_BY_NAME(ControlRigSequencerEditorLibrary)
 
@@ -3000,6 +3002,77 @@ bool UControlRigSequencerEditorLibrary::CollapseControlRigAnimLayers(ULevelSeque
 	return bValid;
 }
 
+UMovieSceneScriptingChannel* UControlRigSequencerEditorLibrary::GetSpaceChannel(UMovieSceneSequence* InSequence, UControlRig* ControlRig, FName ControlName, bool bCreateIfNeeded)
+{
+	UMovieSceneScriptingChannel* ScriptingChannel = nullptr;
+	if (InSequence == nullptr ||  ControlRig == nullptr)
+	{
+		return ScriptingChannel;
+	}
+	UMovieScene* MovieScene = InSequence->GetMovieScene();
+	if (!MovieScene)
+	{
+		return ScriptingChannel;
+
+	}
+
+	bool bFoundTrack = false;
+
+	const TArray<FMovieSceneBinding>& Bindings = ((const UMovieScene*)MovieScene)->GetBindings();
+	bool bRecreateCurves = false;
+	TArray<TPair<UControlRig*, FName>> ControlRigPairsToReselect;
+	for (const FMovieSceneBinding& Binding : Bindings)
+	{
+		TArray<UMovieSceneTrack*> Tracks = MovieScene->FindTracks(UMovieSceneControlRigParameterTrack::StaticClass(), Binding.GetObjectGuid(), NAME_None);
+		for (UMovieSceneTrack* Track : Tracks)
+		{
+			UMovieSceneControlRigParameterTrack* ControlRigParameterTrack = Cast<UMovieSceneControlRigParameterTrack>(Track);
+			if (ControlRigParameterTrack && ControlRigParameterTrack->GetControlRig() == ControlRig)
+			{
+				TArray<UMovieSceneSection*> AllSections = ControlRigParameterTrack->GetAllSections();
+				if (UMovieSceneControlRigParameterSection* ActiveSection = Cast<UMovieSceneControlRigParameterSection>(AllSections.Num() > 0 ? AllSections[0] : nullptr))
+				{
+					FSpaceControlNameAndChannel* NameAndChannel = ActiveSection->GetSpaceChannel(ControlName);
+					if (NameAndChannel == nullptr && bCreateIfNeeded)
+					{
+						ControlRig->Modify();
+						ActiveSection->AddSpaceChannel(ControlName, true /*ReconstructChannelProxy*/);
+						NameAndChannel = ActiveSection->GetSpaceChannel(ControlName);
+						//important we don't want the adding of the space channel to be in the undo since we want the channel to stick around if we get undone
+						ActiveSection->Modify();
+					}
+
+					if (NameAndChannel)
+					{
+						FMovieSceneChannelProxy& ChannelProxy = ActiveSection->GetChannelProxy();
+						TArray<UMovieSceneScriptingChannel*> Channels;
+						GetScriptingChannelsForChannel<FMovieSceneControlRigSpaceChannel, UMovieSceneScriptingSpaceChannel>(ChannelProxy, InSequence, ActiveSection, Channels);
+						if (Channels.Num() > 0)
+						{
+							for (UMovieSceneScriptingChannel* Channel : Channels)
+							{
+								if (Channel)
+								{
+									FString ChannelName = Channel->ChannelName.ToString();
+									FString CName, SpaceName;
+									ChannelName.Split(".", &CName, &SpaceName);
+									if (NameAndChannel->ControlName == FName(*CName))
+									{
+										ScriptingChannel = Channel;
+										break;
+									}
+								}
+							}
+						}
+					}
+				}
+				break;
+			}
+		}
+	}
+	return ScriptingChannel;
+}
+
 bool UControlRigSequencerEditorLibrary::SetControlRigSpace(ULevelSequence* LevelSequence, UControlRig* ControlRig, FName ControlName, const FRigElementKey& InSpaceKey, FFrameNumber InTime, EMovieSceneTimeUnit TimeUnit)
 {
 	TWeakPtr<ISequencer> WeakSequencer = GetSequencerFromAsset();
diff --git a/Engine/Plugins/Animation/ControlRig/Source/ControlRigEditor/Private/Sequencer/Scripting/MovieSceneScriptingConstraint.h b/Engine/Plugins/Animation/ControlRig/Source/ControlRigEditor/Private/Sequencer/Scripting/MovieSceneScriptingConstraint.h
new file mode 100644
--- /dev/null
+++ b/Engine/Plugins/Animation/ControlRig/Source/ControlRigEditor/Private/Sequencer/Scripting/MovieSceneScriptingConstraint.h
@@ -0,0 +1,215 @@
+// Copyright Epic Games, Inc. All Rights Reserved.
+
+#pragma once
+
+#include "Sections/MovieSceneConstrainedSection.h"
+#include "ConstraintChannel.h"
+#include "MovieSceneScriptingChannel.h"
+#include "KeyParams.h"
+#include "MovieSceneTimeUnit.h"
+
+#include "MovieSceneScriptingConstraint.generated.h"
+
+class UMovieSceneScriptingChannel;
+class UMovieSceneScriptingKey;
+template <typename ChannelType, typename ScriptingKeyType, typename ScriptingKeyValueType> struct TMovieSceneScriptingChannel;
+
+/**
+* Exposes a Constraint Channel key to Python/Blueprints.
+* Stores a reference to the data so changes to this class are forwarded onto the underlying data structures.
+*/
+UCLASS(BlueprintType)
+class UMovieSceneScriptingConstraintChannelKey : public UMovieSceneScriptingKey, public TMovieSceneScriptingKey<FMovieSceneConstraintChannel, bool>
+{
+	GENERATED_BODY()
+public:
+	/**
+	* Gets the time for this key from the owning channel.
+	* @param TimeUnit	Should the time be returned in Display Rate frames (possibly with a sub-frame value) or in Tick Resolution with no sub-frame values?
+	* @return			The time of this key which combines both the frame number and the sub-frame it is on. Sub-frame will be zero if you request Tick Resolution.
+	*/
+	UFUNCTION(BlueprintCallable, Category = "Sequencer|Keys", meta = (DisplayName = "Get Time (Constraint Key)"))
+	virtual FFrameTime GetTime(EMovieSceneTimeUnit TimeUnit = EMovieSceneTimeUnit::DisplayRate) const override { return GetTimeFromChannel(KeyHandle, OwningSequence, TimeUnit); }
+
+	/**
+	* Sets the time for this key in the owning channel. Will replace any key that already exists at that frame number in this channel.
+	* @param NewFrameNumber	What frame should this key be moved to? This should be in the time unit specified by TimeUnit.
+	* @param SubFrame		If using Display Rate time, what is the sub-frame this should go to? Clamped [0-1), and ignored with when TimeUnit is set to Tick Resolution.
+	* @param TimeUnit		Should the NewFrameNumber be interpreted as Display Rate frames or in Tick Resolution?
+	*/
+	UFUNCTION(BlueprintCallable, Category = "Sequencer|Keys", meta = (DisplayName = "Set Time (Constraint Key)"))
+	void SetTime(const FFrameNumber& NewFrameNumber, float SubFrame = 0.f, EMovieSceneTimeUnit TimeUnit = EMovieSceneTimeUnit::DisplayRate) { SetTimeInChannel(KeyHandle, OwningSequence, OwningSection, NewFrameNumber, TimeUnit, SubFrame); }
+
+	/**
+	* Gets the value for this key from the owning channel.
+	* @return	The  constraint key.
+	*/
+	UFUNCTION(BlueprintCallable, Category = "Sequencer|Keys", meta = (DisplayName = "Get Value (Constraint Key)"))
+	bool GetValue() const
+	{
+		bool Value = GetValueFromChannel(KeyHandle);
+		return Value;
+	}
+
+	/**
+	* Sets the value for this key, reflecting it in the owning channel.
+	* @param InNewValue	The constraint key
+	*/
+	UFUNCTION(BlueprintCallable, Category = "Sequencer|Keys", meta = (DisplayName = "Set Value (Constraint Key)"))
+	void SetValue(const bool& InNewValue)
+	{
+		SetValueInChannel(KeyHandle, OwningSection, InNewValue);
+	}
+};
+
+UCLASS(BlueprintType)
+class UMovieSceneScriptingConstraintChannel : public UMovieSceneScriptingChannel
+{
+	GENERATED_BODY()
+
+	using Impl = TMovieSceneScriptingChannel<FMovieSceneConstraintChannel, UMovieSceneScriptingConstraintChannelKey, bool>;
+
+public:
+	/**
+	* Add a key to this channel. This initializes a new key and returns a reference to it.
+	* @param	InTime			The frame this key should go on. Respects TimeUnit to determine if it is a display rate frame or a tick resolution frame.
+	* @param	NewValue		The value that this key should be created with.
+	* @param	SubFrame		Optional [0-1) clamped sub-frame to put this key on. Ignored if TimeUnit is set to Tick Resolution.
+	* @param	TimeUnit 		Is the specified InTime in Display Rate frames or Tick Resolution.
+	* @return	The key that was created with the specified values at the specified time.
+	*/
+	UFUNCTION(BlueprintCallable, Category = "Sequencer|Keys", meta = (DisplayName = "Add Key (Constraint Key)"))
+	UMovieSceneScriptingConstraintChannelKey* AddKey(const FFrameNumber InTime, bool NewValue, float SubFrame = 0.f, EMovieSceneTimeUnit TimeUnit = EMovieSceneTimeUnit::DisplayRate)
+	{
+		return Impl::AddKeyInChannel(ChannelHandle, OwningSequence, OwningSection, InTime, NewValue, SubFrame, TimeUnit, EMovieSceneKeyInterpolation::Auto);
+	}
+
+	/**
+	* Removes the specified key. Does nothing if the key is not specified or the key belongs to another channel.
+	*/
+	UFUNCTION(BlueprintCallable, Category = "Sequencer|Keys", meta = (DisplayName = "Remove Key (Constraint Key)"))
+	virtual void RemoveKey(UMovieSceneScriptingKey* Key)
+	{
+		Impl::RemoveKeyFromChannel(ChannelHandle, Key);
+	}
+
+	/**
+	* Gets all of the keys in this channel.
+	* @return	An array of UMovieSceneScriptingConstraintChannelKeys contained by this channel.
+	*			Returns all keys even if clipped by the owning section's boundaries or outside of the current sequence play range.
+	*/
+	UFUNCTION(BlueprintCallable, Category = "Sequencer|Keys", meta = (DisplayName = "Get Keys (Constraint Key)"))
+	virtual TArray<UMovieSceneScriptingKey*> GetKeys() const override
+	{
+		return Impl::GetKeysInChannel(ChannelHandle, OwningSequence, OwningSection);
+	}
+
+	/**
+	* Gets the keys in this channel specified by the specific index
+	* @Indices  The indices from which to get the keys from
+	* @return	An array of UMovieSceneScriptingKey's contained by this channel.
+	*			Returns all keys specified by the indices, even if out of range.
+	*/
+	UFUNCTION(BlueprintCallable, Category = "Sequencer|Keys", meta = (DisplayName = "Get Keys By Index (Constraint Key)"))
+	virtual TArray<UMovieSceneScriptingKey*> GetKeysByIndex(const TArray<int32>& Indices) const override
+	{
+		return Impl::GetKeysInChannelByIndex(ChannelHandle, OwningSequence, OwningSection, Indices);
+	}
+
+	/**
+	* Set this channel's default value that should be used when no keys are present.
+	* Sets bHasDefaultValue to true automatically.
+	*/
+	UFUNCTION(BlueprintCallable, Category = "Sequencer|Keys", meta = (DisplayName = "Set Default (Constraint Key)"))
+	void SetDefault(bool InDefaultValue)
+	{
+		bool ReferenceKey = bool(InDefaultValue);
+		Impl::SetDefaultInChannel(ChannelHandle, OwningSequence, OwningSection, ReferenceKey);
+	}
+
+
+	/**
+	* Get this channel's default value that will be used when no keys are present. Only a valid
+	* value when HasDefault() returns true.
+	*/
+	UFUNCTION(BlueprintCallable, Category = "Sequencer|Keys", meta = (DisplayName = "Get Default (Constraint Key)"))
+	bool GetDefault() const
+	{
+		TOptional<bool> DefaultValue = Impl::GetDefaultFromChannel(ChannelHandle);
+		return DefaultValue.IsSet() ? DefaultValue.GetValue() : bool();
+	}
+
+	/**
+	* Remove this channel's default value causing the channel to have no effect where no keys are present
+	*/
+	UFUNCTION(BlueprintCallable, Category = "Sequencer|Keys", meta = (DisplayName = "Remove Default (Constraint Key)"))
+	void RemoveDefault()
+	{
+		if (OwningSection.IsValid())
+		{
+			OwningSection->Modify();
+		}
+
+		FMovieSceneConstraintChannel* Channel = ChannelHandle.Get();
+		if (Channel)
+		{
+			Channel->ClearDefault();
+
+#if WITH_EDITOR
+			const FMovieSceneChannelMetaData* MetaData = ChannelHandle.GetMetaData();
+			if (MetaData && OwningSection.IsValid() && OwningSequence.IsValid() && OwningSequence->GetMovieScene())
+			{
+				OwningSequence->GetMovieScene()->OnChannelChanged().Broadcast(MetaData, OwningSection.Get());
+			}
+#endif
+			return;
+		}
+		FFrame::KismetExecutionMessage(TEXT("Invalid ChannelHandle for MovieSceneScriptingChannel, failed to remove default value."), ELogVerbosity::Error);
+	}
+
+	/**
+	* @return Does this channel have a default value set?
+	*/
+	UFUNCTION(BlueprintCallable, Category = "Sequencer|Keys", meta = (DisplayName = "Has Default (Constraint Key)"))
+	bool HasDefault() const
+	{
+		TOptional<bool> DefaultValue = Impl::GetDefaultFromChannel(ChannelHandle);
+		return DefaultValue.IsSet();
+	}
+
+	/**
+	 * Transform the keys in time in the channel by an offset, scale and pivot
+	 * @param	OffsetFrame	    The amount to offset the keys by
+	 * @param	Scale	        The amount to scale the keys by
+	 * @param	PivotFrame	    The frame to pivot around when scaling the keys
+	 * @param   ScriptingRange  The time range of the keys to transform
+	 * @param	TimeUnit 		Is the specified OffsetFrame, PivotFrame, and ScriptingRange in Display Rate frames or Tick Resolution.
+	 *
+	 */
+	UFUNCTION(BlueprintCallable, Category = "Sequencer|Keys", meta = (DisplayName = "Transform (Constraint Key)"))
+	void Transform(FFrameNumber OffsetFrame, double Scale, FFrameNumber PivotFrame, FSequencerScriptingRange ScriptingRange, EMovieSceneTimeUnit TimeUnit = EMovieSceneTimeUnit::DisplayRate)
+	{
+		Impl::TransformKeysInChannel(ChannelHandle, OwningSequence, OwningSection, OffsetFrame, Scale, PivotFrame, ScriptingRange, TimeUnit);
+	}
+
+	/**
+	* @return Get the tickable constraint associated with this channel
+	*/
+	UFUNCTION(BlueprintCallable, Category = "Sequencer|Keys", meta = (DisplayName = "Get Constraint (Constraint Channel)"))
+	UTickableConstraint* GetConstraint()
+	{
+		if(IMovieSceneConstrainedSection* Section = Cast<IMovieSceneConstrainedSection>(OwningSection.Get()))
+		{
+			FMovieSceneConstraintChannel* Channel = ChannelHandle.Get();
+			if (Channel)
+			{
+				return Section->FindConstraint(Channel);
+			}
+		}
+		return nullptr;
+	}
+public:
+	TWeakObjectPtr<UMovieSceneSequence> OwningSequence;
+	TWeakObjectPtr<UMovieSceneSection> OwningSection;
+	TMovieSceneChannelHandle<FMovieSceneConstraintChannel> ChannelHandle;
+};
diff --git a/Engine/Plugins/Animation/ControlRig/Source/ControlRigEditor/Private/Sequencer/Scripting/MovieSceneScriptingSpace.h b/Engine/Plugins/Animation/ControlRig/Source/ControlRigEditor/Private/Sequencer/Scripting/MovieSceneScriptingSpace.h
new file mode 100644
--- /dev/null
+++ b/Engine/Plugins/Animation/ControlRig/Source/ControlRigEditor/Private/Sequencer/Scripting/MovieSceneScriptingSpace.h
@@ -0,0 +1,199 @@
+// Copyright Epic Games, Inc. All Rights Reserved.
+
+#pragma once
+
+#include "Sequencer/MovieSceneControlRigSpaceChannel.h"
+#include "MovieSceneScriptingChannel.h"
+#include "KeyParams.h"
+#include "MovieSceneTimeUnit.h"
+
+#include "MovieSceneScriptingSpace.generated.h"
+
+class UMovieSceneScriptingChannel;
+class UMovieSceneScriptingKey;
+template <typename ChannelType, typename ScriptingKeyType, typename ScriptingKeyValueType> struct TMovieSceneScriptingChannel;
+
+/**
+* Exposes a Space Channel key to Python/Blueprints.
+* Stores a reference to the data so changes to this class are forwarded onto the underlying data structures.
+*/
+UCLASS(BlueprintType)
+class UMovieSceneScriptingSpaceChannelKey : public UMovieSceneScriptingKey, public TMovieSceneScriptingKey<FMovieSceneControlRigSpaceChannel, FMovieSceneControlRigSpaceBaseKey>
+{
+	GENERATED_BODY()
+public:
+	/**
+	* Gets the time for this key from the owning channel.
+	* @param TimeUnit	Should the time be returned in Display Rate frames (possibly with a sub-frame value) or in Tick Resolution with no sub-frame values?
+	* @return			The time of this key which combines both the frame number and the sub-frame it is on. Sub-frame will be zero if you request Tick Resolution.
+	*/
+	UFUNCTION(BlueprintCallable, Category = "Sequencer|Keys", meta = (DisplayName = "Get Time (Space Keye)"))
+	virtual FFrameTime GetTime(EMovieSceneTimeUnit TimeUnit = EMovieSceneTimeUnit::DisplayRate) const override { return GetTimeFromChannel(KeyHandle, OwningSequence, TimeUnit); }
+
+	/**
+	* Sets the time for this key in the owning channel. Will replace any key that already exists at that frame number in this channel.
+	* @param NewFrameNumber	What frame should this key be moved to? This should be in the time unit specified by TimeUnit.
+	* @param SubFrame		If using Display Rate time, what is the sub-frame this should go to? Clamped [0-1), and ignored with when TimeUnit is set to Tick Resolution.
+	* @param TimeUnit		Should the NewFrameNumber be interpreted as Display Rate frames or in Tick Resolution?
+	*/
+	UFUNCTION(BlueprintCallable, Category = "Sequencer|Keys", meta = (DisplayName = "Set Time (Space Keye)"))
+	void SetTime(const FFrameNumber& NewFrameNumber, float SubFrame = 0.f, EMovieSceneTimeUnit TimeUnit = EMovieSceneTimeUnit::DisplayRate) { SetTimeInChannel(KeyHandle, OwningSequence, OwningSection, NewFrameNumber, TimeUnit, SubFrame); }
+
+	/**
+	* Gets the value for this key from the owning channel.
+	* @return The space key
+	*/
+	UFUNCTION(BlueprintCallable, Category = "Sequencer|Keys", meta = (DisplayName = "Get Value (Space Key)"))
+	FMovieSceneControlRigSpaceBaseKey GetValue() const
+	{
+		FMovieSceneControlRigSpaceBaseKey Value = GetValueFromChannel(KeyHandle);
+		return Value;
+	}
+
+	/**
+	* Sets the value for this key, reflecting it in the owning channel.
+	* @param InNewValue	The new space key
+	*/
+	UFUNCTION(BlueprintCallable, Category = "Sequencer|Keys", meta = (DisplayName = "Set Value (Space Key)"))
+	void SetValue(const FMovieSceneControlRigSpaceBaseKey& InNewValue)
+	{
+		FMovieSceneControlRigSpaceBaseKey ReferenceKey = FMovieSceneControlRigSpaceBaseKey(InNewValue);
+		SetValueInChannel(KeyHandle, OwningSection, ReferenceKey);
+	}
+};
+
+UCLASS(BlueprintType)
+class UMovieSceneScriptingSpaceChannel: public UMovieSceneScriptingChannel
+{
+	GENERATED_BODY()
+
+	using Impl = TMovieSceneScriptingChannel<FMovieSceneControlRigSpaceChannel, UMovieSceneScriptingSpaceChannelKey, FMovieSceneControlRigSpaceBaseKey>;
+
+public:
+	/**
+	* Add a key to this channel. This initializes a new key and returns a reference to it.
+	* @param	InTime			The frame this key should go on. Respects TimeUnit to determine if it is a display rate frame or a tick resolution frame.
+	* @param	NewValue		The value that this key should be created with.
+	* @param	SubFrame		Optional [0-1) clamped sub-frame to put this key on. Ignored if TimeUnit is set to Tick Resolution.
+	* @param	TimeUnit 		Is the specified InTime in Display Rate frames or Tick Resolution.
+	* @return	The key that was created with the specified values at the specified time.
+	*/
+	UFUNCTION(BlueprintCallable, Category = "Sequencer|Keys", meta = (DisplayName = "Add Key (Space Key)"))
+	UMovieSceneScriptingSpaceChannelKey* AddKey(const FFrameNumber InTime, FMovieSceneControlRigSpaceBaseKey NewValue, float SubFrame = 0.f, EMovieSceneTimeUnit TimeUnit = EMovieSceneTimeUnit::DisplayRate)
+	{
+		return Impl::AddKeyInChannel(ChannelHandle, OwningSequence, OwningSection, InTime, NewValue, SubFrame, TimeUnit, EMovieSceneKeyInterpolation::Auto);
+	}
+
+	/**
+	* Removes the specified key. Does nothing if the key is not specified or the key belongs to another channel.
+	*/
+	UFUNCTION(BlueprintCallable, Category = "Sequencer|Keys", meta = (DisplayName = "Remove Key (Space Key)"))
+	virtual void RemoveKey(UMovieSceneScriptingKey* Key)
+	{
+		Impl::RemoveKeyFromChannel(ChannelHandle, Key);
+	}
+
+	/**
+	* Gets all of the keys in this channel.
+	* @return	An array of UMovieSceneScriptingSpaceChannelKeys contained by this channel.
+	*			Returns all keys even if clipped by the owning section's boundaries or outside of the current sequence play range.
+	*/
+	UFUNCTION(BlueprintCallable, Category = "Sequencer|Keys", meta = (DisplayName = "Get Keys (Space Key)"))
+	virtual TArray<UMovieSceneScriptingKey*> GetKeys() const override
+	{
+		return Impl::GetKeysInChannel(ChannelHandle, OwningSequence, OwningSection);
+	}
+
+	/**
+	* Gets the keys in this channel specified by the specific index
+	* @Indices  The indices from which to get the keys from
+	* @return	An array of UMovieSceneScriptingKey's contained by this channel.
+	*			Returns all keys specified by the indices, even if out of range.
+	*/
+	UFUNCTION(BlueprintCallable, Category = "Sequencer|Keys", meta = (DisplayName = "Get Keys By Index (Space Key)"))
+	virtual TArray<UMovieSceneScriptingKey*> GetKeysByIndex(const TArray<int32>& Indices) const override
+	{
+		return Impl::GetKeysInChannelByIndex(ChannelHandle, OwningSequence, OwningSection, Indices);
+	}
+
+	/**
+	* Set this channel's default value that should be used when no keys are present.
+	* Sets bHasDefaultValue to true automatically.
+	*/
+	UFUNCTION(BlueprintCallable, Category = "Sequencer|Keys", meta = (DisplayName = "Set Default (Space Key)"))
+	void SetDefault(FMovieSceneControlRigSpaceBaseKey InDefaultValue)
+	{
+		FMovieSceneControlRigSpaceBaseKey ReferenceKey = FMovieSceneControlRigSpaceBaseKey(InDefaultValue);
+		Impl::SetDefaultInChannel(ChannelHandle, OwningSequence, OwningSection, ReferenceKey);
+	}
+
+
+	/**
+	* Get this channel's default value that will be used when no keys are present. Only a valid
+	* value when HasDefault() returns true.
+	*/
+	UFUNCTION(BlueprintCallable, Category = "Sequencer|Keys", meta = (DisplayName = "Get Default (Space Key)"))
+	FMovieSceneControlRigSpaceBaseKey GetDefault() const
+	{
+		TOptional<FMovieSceneControlRigSpaceBaseKey> DefaultValue = Impl::GetDefaultFromChannel(ChannelHandle);
+		return DefaultValue.IsSet() ? DefaultValue.GetValue() : FMovieSceneControlRigSpaceBaseKey();
+	}
+	
+	/**
+	* Remove this channel's default value causing the channel to have no effect where no keys are present
+	*/
+	UFUNCTION(BlueprintCallable, Category = "Sequencer|Keys", meta = (DisplayName = "Remove Default (Space Key)"))
+	void RemoveDefault()
+	{
+		if (OwningSection.IsValid())
+		{
+			OwningSection->Modify();
+		}
+
+		FMovieSceneControlRigSpaceChannel* Channel = ChannelHandle.Get();
+		if (Channel)
+		{
+			Channel->ClearDefault();
+
+#if WITH_EDITOR
+			const FMovieSceneChannelMetaData* MetaData = ChannelHandle.GetMetaData();
+			if (MetaData && OwningSection.IsValid() && OwningSequence.IsValid() && OwningSequence->GetMovieScene())
+			{
+				OwningSequence->GetMovieScene()->OnChannelChanged().Broadcast(MetaData, OwningSection.Get());
+			}
+#endif
+			return;
+		}
+		FFrame::KismetExecutionMessage(TEXT("Invalid ChannelHandle for MovieSceneScriptingChannel, failed to remove default value."), ELogVerbosity::Error);
+	}
+
+	/**
+	* @return Does this channel have a default value set?
+	*/
+	UFUNCTION(BlueprintCallable, Category = "Sequencer|Keys", meta = (DisplayName = "Has Default (Space Key)"))
+	bool HasDefault() const
+	{
+		TOptional<FMovieSceneControlRigSpaceBaseKey> DefaultValue = Impl::GetDefaultFromChannel(ChannelHandle);
+		return DefaultValue.IsSet();
+	}
+
+	/**
+	 * Transform the keys in time in the channel by an offset, scale and pivot
+	 * @param	OffsetFrame	    The amount to offset the keys by
+	 * @param	Scale	        The amount to scale the keys by
+	 * @param	PivotFrame	    The frame to pivot around when scaling the keys
+	 * @param   ScriptingRange  The time range of the keys to transform
+	 * @param	TimeUnit 		Is the specified OffsetFrame, PivotFrame, and ScriptingRange in Display Rate frames or Tick Resolution.
+	 *
+	 */
+	UFUNCTION(BlueprintCallable, Category = "Sequencer|Keys", meta = (DisplayName = "Transform (Space Key)"))
+	void Transform(FFrameNumber OffsetFrame, double Scale, FFrameNumber PivotFrame, FSequencerScriptingRange ScriptingRange, EMovieSceneTimeUnit TimeUnit = EMovieSceneTimeUnit::DisplayRate)
+	{
+		Impl::TransformKeysInChannel(ChannelHandle, OwningSequence, OwningSection, OffsetFrame, Scale, PivotFrame, ScriptingRange, TimeUnit);
+	}
+
+public:
+	TWeakObjectPtr<UMovieSceneSequence> OwningSequence;
+	TWeakObjectPtr<UMovieSceneSection> OwningSection;
+	TMovieSceneChannelHandle<FMovieSceneControlRigSpaceChannel> ChannelHandle;
+};
diff --git a/Engine/Plugins/Animation/ControlRig/Source/ControlRigEditor/Public/ControlRigSequencerEditorLibrary.h b/Engine/Plugins/Animation/ControlRig/Source/ControlRigEditor/Public/ControlRigSequencerEditorLibrary.h
--- a/Engine/Plugins/Animation/ControlRig/Source/ControlRigEditor/Public/ControlRigSequencerEditorLibrary.h
+++ b/Engine/Plugins/Animation/ControlRig/Source/ControlRigEditor/Public/ControlRigSequencerEditorLibrary.h
@@ -18,6 +18,9 @@
 #include "MovieSceneTimeUnit.h"
 #include "RigSpacePickerBakeSettings.h"
 #include "Filters/CurveEditorSmartReduceFilter.h"
+
+#include "MovieSceneScriptingChannel.h"
+
 #include "ControlRigSequencerEditorLibrary.generated.h"
 
 #define UE_API CONTROLRIGEDITOR_API
@@ -1021,6 +1024,18 @@ public:
 	UFUNCTION(BlueprintCallable, Category = "Editor Scripting | Sequencer Tools | Control Rig")
 	static UE_API bool CollapseControlRigAnimLayersWithSettings(ULevelSequence* InSequence, UMovieSceneControlRigParameterTrack* InTrack, const FBakingAnimationKeySettings& InSettings);
 
+	/*
+	 * Get the space channel for the specified control rig and control name. Will create it if needed. Note spaces will always live on the first section.
+	 *
+	 * @param InSequence Sequence that to get space channel
+	 * @param InControlRig ControlRig with the Control
+	 * @param InControlName Name of control to have space channel
+	 * @param bCreateIfNeeded Create the channel if it doesn't exist
+	 * @return return the space channel if one exists or was created based upon bCreateIfNeeded
+	 */
+	UFUNCTION(BlueprintCallable, Category = "Editor Scripting | Sequencer Tools | Control Rig")
+	static UMovieSceneScriptingChannel* GetSpaceChannel(UMovieSceneSequence* InSequence, UControlRig* InControlRig, FName InControlName, bool bCreateIfNeeded = true);
+
 	/*
 	 * Get the default parent key, can be used a parent space.
 	 *
@@ -1036,7 +1051,7 @@ public:
 	static UE_API FRigElementKey GetWorldSpaceReferenceKey();
 
 	/*
-	 * Set the a key for the Control Rig Space for the Control at the specified time. If space is the same as the current no key witll be set.
+	 * Set the a key for the Control Rig Space for the Control at the specified time. If space is the same as the current no key will be set.
 	 *
 	 * @param InSequence Sequence to set the space
 	 * @param InControlRig ControlRig with the Control
diff --git a/Engine/Plugins/MovieScene/SequencerScripting/Source/SequencerScripting/Private/ExtensionLibraries/MovieSceneSectionExtensions.cpp b/Engine/Plugins/MovieScene/SequencerScripting/Source/SequencerScripting/Private/ExtensionLibraries/MovieSceneSectionExtensions.cpp
--- a/Engine/Plugins/MovieScene/SequencerScripting/Source/SequencerScripting/Private/ExtensionLibraries/MovieSceneSectionExtensions.cpp
+++ b/Engine/Plugins/MovieScene/SequencerScripting/Source/SequencerScripting/Private/ExtensionLibraries/MovieSceneSectionExtensions.cpp
@@ -1,8 +1,6 @@
 // Copyright Epic Games, Inc. All Rights Reserved.
 
 #include "ExtensionLibraries/MovieSceneSectionExtensions.h"
-#include "Channels/MovieSceneChannelProxy.h"
-#include "Channels/MovieSceneChannelHandle.h"
 #include "Evaluation/MovieSceneSequenceTransform.h"
 #include "KeysAndChannels/MovieSceneScriptingBool.h"
 #include "Evaluation/MovieSceneTimeTransform.h"
@@ -446,11 +444,6 @@ void UMovieSceneSectionExtensions::SetEndFrameBounded(UMovieSceneSection* Sectio
 	}
 }
 
-template<typename ChannelType, typename ScriptingChannelType>
-void SetScriptingChannelHandle(ScriptingChannelType* ScriptingChannel, FMovieSceneChannelProxy& ChannelProxy, int32 ChannelIndex)
-{
-	ScriptingChannel->ChannelHandle = ChannelProxy.MakeHandle<ChannelType>(ChannelIndex);
-}
 template<>
 void SetScriptingChannelHandle<FMovieSceneFloatChannel, UMovieSceneScriptingFloatChannel>(UMovieSceneScriptingFloatChannel* ScriptingChannel, FMovieSceneChannelProxy& ChannelProxy, int32 ChannelIndex)
 {
@@ -462,47 +455,21 @@ void SetScriptingChannelHandle<FMovieSceneDoubleChannel, UMovieSceneScriptingFlo
 	ScriptingChannel->DoubleChannelHandle = ChannelProxy.MakeHandle<FMovieSceneDoubleChannel>(ChannelIndex);
 }
 
-template<typename ChannelType, typename ScriptingChannelType>
-void GetScriptingChannelsForChannel(FMovieSceneChannelProxy& ChannelProxy,TWeakObjectPtr<UMovieSceneSequence> Sequence, TWeakObjectPtr<UMovieSceneSection> Section, TArray<UMovieSceneScriptingChannel*>& OutChannels)
+
+
+TMap<FGuid,TFunction<bool(UMovieSceneSection* Section, TSubclassOf<UMovieSceneScriptingChannel> ChannelType, TArray<UMovieSceneScriptingChannel*>& OutChannels)>>  UMovieSceneSectionExtensions::CustomChannelsCallbacks;
+
+FGuid UMovieSceneSectionExtensions::SetGetCustomChannelsFunction(TFunction<bool(UMovieSceneSection* Section, TSubclassOf<UMovieSceneScriptingChannel> ChannelType, TArray<UMovieSceneScriptingChannel*>& OutChannels) > CustommFunc)
 {
-	const FMovieSceneChannelEntry* Entry = ChannelProxy.FindEntry(ChannelType::StaticStruct()->GetFName());
-	if (Entry)
-	{
-#if WITH_EDITORONLY_DATA
-		TArrayView<const FMovieSceneChannelMetaData> MetaData = Entry->GetMetaData();
-		for (int32 Index = 0; Index < MetaData.Num(); ++Index)
-		{
-			if (MetaData[Index].bEnabled)
-			{
-				const FName ChannelName = MetaData[Index].Name;
-				FName UniqueChannelName = MakeUniqueObjectName(GetTransientPackage(), ScriptingChannelType::StaticClass(), ChannelName);
-				ScriptingChannelType* ScriptingChannel = NewObject<ScriptingChannelType>(GetTransientPackage(), UniqueChannelName);
-				SetScriptingChannelHandle<ChannelType, ScriptingChannelType>(ScriptingChannel, ChannelProxy, Index);
-				ScriptingChannel->OwningSection = Section;
-				ScriptingChannel->OwningSequence = Sequence;
-				ScriptingChannel->ChannelName = ChannelName;
-
-				OutChannels.Add(ScriptingChannel);
-			}
-		}
-#endif
-	}
+	FGuid NewGuid = FGuid::NewGuid();
+	CustomChannelsCallbacks.Add(NewGuid,CustommFunc);
+	return NewGuid;
 }
-
-template<typename ChannelType, typename ScriptingChannelType>
-UMovieSceneScriptingChannel* GetScriptingChannelForChannel(const FMovieSceneChannel* MovieSceneChannel, FMovieSceneChannelProxy& ChannelProxy, const FMovieSceneChannelHandle& ChanneHandle, TWeakObjectPtr<UMovieSceneSequence> Sequence, TWeakObjectPtr<UMovieSceneSection> Section, const FName& ChannelName)
+void UMovieSceneSectionExtensions::RemoveGetCustomChannelsFunction(const FGuid& InGuid)
 {
-	FName UniqueChannelName = MakeUniqueObjectName(GetTransientPackage(), ScriptingChannelType::StaticClass(), ChannelName);
-	ScriptingChannelType* ScriptingChannel = NewObject<ScriptingChannelType>(GetTransientPackage(), UniqueChannelName);
-	int32 Index = ChannelProxy.FindIndex(ChanneHandle.GetChannelTypeName(), MovieSceneChannel);
-	SetScriptingChannelHandle<ChannelType, ScriptingChannelType>(ScriptingChannel, ChannelProxy, Index);
-	ScriptingChannel->OwningSection = Section;
-	ScriptingChannel->OwningSequence = Sequence;
-	ScriptingChannel->ChannelName = ChannelName;
-	return ScriptingChannel;
+	CustomChannelsCallbacks.Remove(InGuid);
 }
 
-
 TArray<UMovieSceneScriptingChannel*> UMovieSceneSectionExtensions::GetAllChannels(UMovieSceneSection* Section)
 {
 	TArray<UMovieSceneScriptingChannel*> Channels;
@@ -528,7 +495,13 @@ TArray<UMovieSceneScriptingChannel*> UMovieSceneSectionExtensions::GetAllChannel
 	GetScriptingChannelsForChannel<FMovieSceneActorReferenceData, UMovieSceneScriptingActorReferenceChannel>(ChannelProxy, Sequence, Section, Channels);
 	GetScriptingChannelsForChannel<FMovieSceneObjectPathChannel, UMovieSceneScriptingObjectPathChannel>(ChannelProxy, Sequence, Section, Channels);
 	GetScriptingChannelsForChannel<FMovieSceneParticleChannel, UMovieSceneScriptingParticleChannel>(ChannelProxy, Sequence, Section, Channels);
+	
 
+	for (const TPair<FGuid,TFunction<bool(UMovieSceneSection* Section, TSubclassOf<UMovieSceneScriptingChannel> ChannelType, TArray<UMovieSceneScriptingChannel*>& OutChannels)>>&
+		Function : CustomChannelsCallbacks)
+	{
+		Function.Value(Section, UMovieSceneScriptingChannel::StaticClass(), Channels);
+	}
 	return Channels;
 }
 
@@ -546,6 +519,18 @@ TArray<UMovieSceneScriptingChannel*> UMovieSceneSectionExtensions::GetChannelsBy
 	// ToDo: This needs to be done dynamically in the future but there's not a good way to do that right now with all of the SFINAE-based Templating driving these channels
 	TWeakObjectPtr<UMovieSceneSequence> Sequence = Section->GetTypedOuter<UMovieSceneSequence>();
 
+	//But we added this to be able to specify custom UMoiveSceneScriptingChanngels
+	for (const TPair<FGuid, TFunction<bool(UMovieSceneSection* Section, TSubclassOf<UMovieSceneScriptingChannel> ChannelType, TArray<UMovieSceneScriptingChannel*>& OutChannels)>>&
+		Function : CustomChannelsCallbacks)
+	{
+		if (Function.Value(Section, ChannelType, Channels))
+		{
+			return Channels;
+		}
+	}
+
+
+
 	if (ChannelType == UMovieSceneScriptingBoolChannel::StaticClass())					{ GetScriptingChannelsForChannel<FMovieSceneBoolChannel, UMovieSceneScriptingBoolChannel>(ChannelProxy, Sequence, Section, Channels); }
 	else if (ChannelType == UMovieSceneScriptingByteChannel::StaticClass())				{ GetScriptingChannelsForChannel<FMovieSceneByteChannel, UMovieSceneScriptingByteChannel>(ChannelProxy, Sequence, Section, Channels); }
 	else if (ChannelType == UMovieSceneScriptingFloatChannel::StaticClass())			{ GetScriptingChannelsForChannel<FMovieSceneFloatChannel, UMovieSceneScriptingFloatChannel>(ChannelProxy, Sequence, Section, Channels); }
diff --git a/Engine/Plugins/MovieScene/SequencerScripting/Source/SequencerScripting/Private/KeysAndChannels/MovieSceneScriptingChannel.h b/Engine/Plugins/MovieScene/SequencerScripting/Source/SequencerScripting/Private/KeysAndChannels/MovieSceneScriptingChannel.h
deleted file mode 100644
diff --git a/Engine/Plugins/MovieScene/SequencerScripting/Source/SequencerScripting/Public/ExtensionLibraries/MovieSceneSectionExtensions.h b/Engine/Plugins/MovieScene/SequencerScripting/Source/SequencerScripting/Public/ExtensionLibraries/MovieSceneSectionExtensions.h
--- a/Engine/Plugins/MovieScene/SequencerScripting/Source/SequencerScripting/Public/ExtensionLibraries/MovieSceneSectionExtensions.h
+++ b/Engine/Plugins/MovieScene/SequencerScripting/Source/SequencerScripting/Public/ExtensionLibraries/MovieSceneSectionExtensions.h
@@ -3,6 +3,9 @@
 #pragma once
 
 #include "Kismet/BlueprintFunctionLibrary.h"
+
+#include "Channels/MovieSceneChannelProxy.h"
+#include "Channels/MovieSceneChannelHandle.h"
 #include "MovieSceneSectionExtensions.generated.h"
 
 class UMovieSceneScriptingChannel;
@@ -257,7 +260,73 @@ public:
 	*/
 	static UE_API FMovieSceneChannel* GetMovieSceneChannel(UMovieSceneSection* InSection, const FName& InName);
 
+	/**
+	* Function to add a channel type for a section
+	* @param CustommFunc function to get called when asking for custom channels
+	* @return the Guid used to remove it later
+	**/
+	static UE_API FGuid SetGetCustomChannelsFunction(TFunction<bool(UMovieSceneSection* Section, TSubclassOf<UMovieSceneScriptingChannel> ChannelType, TArray<UMovieSceneScriptingChannel*>& OutChannels)> CustommFunc);
+
+	/**
+	* Function to remove custom channels function
+	* @param InGuid guid to remove that was returned with SetGetCustomChannelsFunction
+	**/
+	static UE_API void RemoveGetCustomChannelsFunction(const FGuid& InGuid);
 
+private:
+
+	static TMap<FGuid,TFunction<bool(UMovieSceneSection* Section, TSubclassOf<UMovieSceneScriptingChannel> ChannelType, TArray<UMovieSceneScriptingChannel*>& OutChannels)>> CustomChannelsCallbacks;
 };
 
+
+/**
+*   Set of template functions that will help with making custom channel types
+*/
+
+template<typename ChannelType, typename ScriptingChannelType>
+void SetScriptingChannelHandle(ScriptingChannelType* ScriptingChannel, FMovieSceneChannelProxy& ChannelProxy, int32 ChannelIndex)
+{
+	ScriptingChannel->ChannelHandle = ChannelProxy.MakeHandle<ChannelType>(ChannelIndex);
+}
+
+template<typename ChannelType, typename ScriptingChannelType>
+void GetScriptingChannelsForChannel(FMovieSceneChannelProxy& ChannelProxy, TWeakObjectPtr<UMovieSceneSequence> Sequence, TWeakObjectPtr<UMovieSceneSection> Section, TArray<UMovieSceneScriptingChannel*>& OutChannels)
+{
+	const FMovieSceneChannelEntry* Entry = ChannelProxy.FindEntry(ChannelType::StaticStruct()->GetFName());
+	if (Entry)
+	{
+#if WITH_EDITORONLY_DATA
+		TArrayView<const FMovieSceneChannelMetaData> MetaData = Entry->GetMetaData();
+		for (int32 Index = 0; Index < MetaData.Num(); ++Index)
+		{
+			if (MetaData[Index].bEnabled)
+			{
+				const FName ChannelName = MetaData[Index].Name;
+				FName UniqueChannelName = MakeUniqueObjectName(GetTransientPackage(), ScriptingChannelType::StaticClass(), ChannelName);
+				ScriptingChannelType* ScriptingChannel = NewObject<ScriptingChannelType>(GetTransientPackage(), UniqueChannelName);
+				SetScriptingChannelHandle<ChannelType, ScriptingChannelType>(ScriptingChannel, ChannelProxy, Index);
+				ScriptingChannel->OwningSection = Section;
+				ScriptingChannel->OwningSequence = Sequence;
+				ScriptingChannel->ChannelName = ChannelName;
+
+				OutChannels.Add(ScriptingChannel);
+			}
+		}
+#endif
+	}
+}
+
+template<typename ChannelType, typename ScriptingChannelType>
+UMovieSceneScriptingChannel* GetScriptingChannelForChannel(const FMovieSceneChannel* MovieSceneChannel, FMovieSceneChannelProxy& ChannelProxy, const FMovieSceneChannelHandle& ChanneHandle, TWeakObjectPtr<UMovieSceneSequence> Sequence, TWeakObjectPtr<UMovieSceneSection> Section, const FName& ChannelName)
+{
+	FName UniqueChannelName = MakeUniqueObjectName(GetTransientPackage(), ScriptingChannelType::StaticClass(), ChannelName);
+	ScriptingChannelType* ScriptingChannel = NewObject<ScriptingChannelType>(GetTransientPackage(), UniqueChannelName);
+	int32 Index = ChannelProxy.FindIndex(ChanneHandle.GetChannelTypeName(), MovieSceneChannel);
+	SetScriptingChannelHandle<ChannelType, ScriptingChannelType>(ScriptingChannel, ChannelProxy, Index);
+	ScriptingChannel->OwningSection = Section;
+	ScriptingChannel->OwningSequence = Sequence;
+	ScriptingChannel->ChannelName = ChannelName;
+	return ScriptingChannel;
+}
+
 #undef UE_API
diff --git a/Engine/Plugins/MovieScene/SequencerScripting/Source/SequencerScripting/Public/MovieSceneScriptingChannel.h b/Engine/Plugins/MovieScene/SequencerScripting/Source/SequencerScripting/Public/MovieSceneScriptingChannel.h
new file mode 100644
--- /dev/null
+++ b/Engine/Plugins/MovieScene/SequencerScripting/Source/SequencerScripting/Public/MovieSceneScriptingChannel.h
@@ -0,0 +1,524 @@
+// Copyright Epic Games, Inc. All Rights Reserved.
+
+#pragma once
+
+#include "Channels/MovieSceneChannelHandle.h"
+#include "Channels/MovieSceneChannelTraits.h"
+#include "SequencerScriptingRange.h"
+#include "MovieSceneTimeUnit.h"
+#include "ExtensionLibraries/MovieSceneSequenceExtensions.h"
+#include "IMovieSceneRetimingInterface.h"
+#include "MovieSceneSection.h"
+#include "MovieSceneSequence.h"
+#include "MovieScene.h"
+#include "MovieSceneScriptingChannel.generated.h"
+
+#define UE_API SEQUENCERSCRIPTING_API
+
+UCLASS(abstract, BlueprintType, MinimalAPI)
+class UMovieSceneScriptingKey : public UObject
+{
+	GENERATED_BODY()
+
+public:
+	/**
+
+	* Gets the time for this key from the owning channel.
+	* @param TimeUnit (Optional) Should the time be returned in Display Rate frames (possibly with a sub-frame value) or in Tick Resolution with no sub-frame values? Defaults to Display Rate.
+	* @return					 The FrameTime of this key which combines both the frame number and the sub-frame it is on. Sub-frame will be zero if you request Tick Resolution.
+	*/
+	virtual FFrameTime GetTime(EMovieSceneTimeUnit TimeUnit = EMovieSceneTimeUnit::DisplayRate) const PURE_VIRTUAL(UMovieSceneScriptingKey::GetTime, return FFrameTime(););
+public:
+	FKeyHandle KeyHandle;
+	TWeakObjectPtr<UMovieSceneSequence> OwningSequence;
+	TWeakObjectPtr<UMovieSceneSection> OwningSection;
+};
+
+
+UCLASS(abstract, BlueprintType, MinimalAPI)
+class UMovieSceneScriptingChannel : public UObject
+{
+	GENERATED_BODY()
+
+public:
+	/**
+	* Gets all of the keys in this channel.
+	* @return	An array of UMovieSceneScriptingKey's contained by this channel.
+	*			Returns all keys even if clipped by the owning section's boundaries or outside of the current sequence play range.
+	*/
+	virtual TArray<UMovieSceneScriptingKey*> GetKeys() const PURE_VIRTUAL(UMovieSceneScriptingChannel::GetKeys, return TArray<UMovieSceneScriptingKey*>(););
+
+	/**
+	* Gets the keys in this channel specified by the specific index
+	* @Indices  The indices from which to get the keys from
+	* @return	An array of UMovieSceneScriptingKey's contained by this channel.
+	*			Returns all keys specified by the indices, even if out of range.
+	*/
+	virtual TArray<UMovieSceneScriptingKey*> GetKeysByIndex(const TArray<int32> &Indices) const PURE_VIRTUAL(UMovieSceneScriptingChannel::GetKeysByIndex, return TArray<UMovieSceneScriptingKey*>(););
+
+
+	UPROPERTY(BlueprintReadOnly, Category="Sequencer|Keys")
+	FName ChannelName;
+};
+
+template<typename ChannelType, typename ScriptingKeyType, typename ScriptingKeyValueType>
+struct TMovieSceneScriptingChannel
+{
+	static ScriptingKeyType* AddKeyInChannel(TMovieSceneChannelHandle<ChannelType> ChannelHandle, TWeakObjectPtr<UMovieSceneSequence> Sequence, 
+		TWeakObjectPtr<UMovieSceneSection> Section, const FFrameNumber InTime, ScriptingKeyValueType& NewValue, float SubFrame, EMovieSceneTimeUnit TimeUnit, EMovieSceneKeyInterpolation Interpolation)
+	{
+		ChannelType* Channel = ChannelHandle.Get();
+		if (Channel)
+		{
+			ScriptingKeyType* Key = NewObject<ScriptingKeyType>();
+
+			// The Key's time is always going to be in Tick Resolution space, but the user may want to set it via Display Rate, so we convert.
+			FFrameNumber KeyTime = InTime;
+			if (TimeUnit == EMovieSceneTimeUnit::DisplayRate)
+			{
+				KeyTime = FFrameRate::TransformTime(KeyTime, UMovieSceneSequenceExtensions::GetDisplayRate(Sequence.Get()), UMovieSceneSequenceExtensions::GetTickResolution(Sequence.Get())).RoundToFrame();
+			}
+
+			if (Section.IsValid())
+			{
+				Section->Modify();
+			}
+
+			using namespace UE::MovieScene;
+			Key->KeyHandle = AddKeyToChannel(Channel, KeyTime, NewValue, Interpolation);
+			Key->ChannelHandle = ChannelHandle;
+			Key->OwningSection = Section;
+			Key->OwningSequence = Sequence;
+
+#if WITH_EDITOR
+			const FMovieSceneChannelMetaData* MetaData = ChannelHandle.GetMetaData();
+			if (MetaData && Section.IsValid() && Sequence.IsValid() && Sequence->GetMovieScene())
+			{
+				Sequence->GetMovieScene()->OnChannelChanged().Broadcast(MetaData, Section.Get());
+			}
+#endif
+
+			return Key;
+		}
+
+		FFrame::KismetExecutionMessage(TEXT("Invalid ChannelHandle for MovieSceneScriptingChannel, failed to add key."), ELogVerbosity::Error);
+		return nullptr;
+	}
+
+	static void RemoveKeyFromChannel(TMovieSceneChannelHandle<ChannelType> ChannelHandle, UMovieSceneScriptingKey* Key)
+	{
+		if (!Key)
+		{
+			return;
+		}
+
+		ChannelType* Channel = ChannelHandle.Get();
+		if (Channel)
+		{
+			if (Key->OwningSection.IsValid())
+			{
+				Key->OwningSection->Modify();
+			}
+
+			TArrayView<FKeyHandle, int> KeyArray = TArrayView<FKeyHandle, int>(&Key->KeyHandle, 1);
+			Channel->DeleteKeys(KeyArray);
+#if WITH_EDITOR
+			const FMovieSceneChannelMetaData* MetaData = ChannelHandle.GetMetaData();
+			if (MetaData && Key->OwningSection.IsValid() && Key->OwningSequence.IsValid() && Key->OwningSequence->GetMovieScene())
+			{
+				Key->OwningSequence->GetMovieScene()->OnChannelChanged().Broadcast(MetaData, Key->OwningSection.Get());
+			}
+#endif
+			return;
+		}
+
+		FFrame::KismetExecutionMessage(TEXT("Invalid ChannelHandle for MovieSceneScriptingChannel, failed to remove key."), ELogVerbosity::Error);
+	}
+
+	static TArray<UMovieSceneScriptingKey*> GetKeysInChannelByIndex(TMovieSceneChannelHandle<ChannelType> ChannelHandle, TWeakObjectPtr<UMovieSceneSequence> Sequence, TWeakObjectPtr<UMovieSceneSection> Section,
+		const TArray<int32>& Indices)
+	{
+		TArray<UMovieSceneScriptingKey*> OutScriptingKeys;
+		ChannelType* Channel = ChannelHandle.Get();
+		if (Channel)
+		{
+			TArray<FFrameNumber> OutTimes;
+			TArray<FKeyHandle> OutKeys;
+			Channel->GetKeys(TRange<FFrameNumber>(), &OutTimes, &OutKeys);
+
+			for (int32 Index: Indices)
+			{
+				if (Index >= 0 && Index < OutTimes.Num())
+				{
+					ScriptingKeyType * Key = NewObject<ScriptingKeyType>();
+					Key->KeyHandle = OutKeys[Index];
+					Key->ChannelHandle = ChannelHandle;
+					Key->OwningSequence = Sequence;
+					Key->OwningSection = Section;
+					OutScriptingKeys.Add(Key);
+				}
+				else
+				{
+					FFrame::KismetExecutionMessage(TEXT("Invalid index for MovieSceneScriptingChannel, failed to get keys by index."), ELogVerbosity::Error);
+				}
+			}
+		}
+		else
+		{
+			FFrame::KismetExecutionMessage(TEXT("Invalid ChannelHandle for MovieSceneScriptingChannel, failed to get keys."), ELogVerbosity::Error);
+		}
+
+		return OutScriptingKeys;
+	}
+
+	static TArray<UMovieSceneScriptingKey*> GetKeysInChannel(TMovieSceneChannelHandle<ChannelType> ChannelHandle, TWeakObjectPtr<UMovieSceneSequence> Sequence, TWeakObjectPtr<UMovieSceneSection> Section)
+	{
+		TArray<UMovieSceneScriptingKey*> OutScriptingKeys;
+		ChannelType* Channel = ChannelHandle.Get();
+		if (Channel)
+		{
+			TArray<FFrameNumber> OutTimes;
+			TArray<FKeyHandle> OutKeys;
+			Channel->GetKeys(TRange<FFrameNumber>(), &OutTimes, &OutKeys);
+
+			for (int32 i = 0; i < OutTimes.Num(); i++)
+			{
+				ScriptingKeyType* Key = NewObject<ScriptingKeyType>();
+				Key->KeyHandle = OutKeys[i];
+				Key->ChannelHandle = ChannelHandle;
+				Key->OwningSequence = Sequence;
+				Key->OwningSection = Section;
+				OutScriptingKeys.Add(Key);
+			}
+		}
+		else
+		{
+			FFrame::KismetExecutionMessage(TEXT("Invalid ChannelHandle for MovieSceneScriptingChannel, failed to get keys."), ELogVerbosity::Error);
+		}
+
+		return OutScriptingKeys;
+	}
+
+	static TArray<ScriptingKeyValueType> EvaluateKeysInChannel(TMovieSceneChannelHandle<ChannelType> ChannelHandle, TWeakObjectPtr<UMovieSceneSequence> Sequence, FSequencerScriptingRange ScriptingRange, FFrameRate FrameRate)
+	{
+		TArray<ScriptingKeyValueType> OutValues;
+		ChannelType* Channel = ChannelHandle.Get();
+		if (Channel)
+		{
+			UMovieSceneSequence* MovieSceneSequence = Sequence.Get();
+			FFrameRate Resolution = UMovieSceneSequenceExtensions::GetTickResolution(MovieSceneSequence);
+			TRange<FFrameNumber> SpecifiedRange = ScriptingRange.ToNative(Resolution);
+			if (SpecifiedRange.HasLowerBound() && SpecifiedRange.HasUpperBound())
+			{
+				FFrameTime Interval = FFrameRate::TransformTime(1, FrameRate, Resolution);
+				FFrameNumber InFrame = UE::MovieScene::DiscreteInclusiveLower(SpecifiedRange);
+				FFrameNumber OutFrame = UE::MovieScene::DiscreteExclusiveUpper(SpecifiedRange);
+				for (FFrameTime EvalTime = InFrame; EvalTime < OutFrame; EvalTime += Interval)
+				{
+					FFrameNumber KeyTime = FFrameRate::Snap(EvalTime, Resolution, FrameRate).FloorToFrame();
+					Channel->Evaluate(KeyTime, OutValues.Emplace_GetRef());
+				}
+			}
+			else
+			{
+				FFrame::KismetExecutionMessage(TEXT("Unbounded range passed to evaluate keys."), ELogVerbosity::Error);
+			}
+		}
+		else
+		{
+			FFrame::KismetExecutionMessage(TEXT("Invalid ChannelHandle for MovieSceneScriptingChannel, failed to evaluate keys."), ELogVerbosity::Error);
+		}
+
+		return OutValues;
+	}
+
+	static FSequencerScriptingRange ComputeEffectiveRangeInChannel(TMovieSceneChannelHandle<ChannelType> ChannelHandle, TWeakObjectPtr<UMovieSceneSequence> Sequence)
+	{
+		FSequencerScriptingRange ScriptingRange;
+		ChannelType* Channel = ChannelHandle.Get();
+		if (Channel)
+		{
+			ScriptingRange = FSequencerScriptingRange::FromNative(Channel->ComputeEffectiveRange(), UMovieSceneSequenceExtensions::GetTickResolution(Sequence.Get()));
+		}
+		else
+		{
+			FFrame::KismetExecutionMessage(TEXT("Invalid ChannelHandle for MovieSceneScriptingChannel, failed to get effective range."), ELogVerbosity::Error);
+		}
+
+		return ScriptingRange;
+	}
+
+	static void SetDefaultInChannel(TMovieSceneChannelHandle<ChannelType> ChannelHandle, TWeakObjectPtr<UMovieSceneSequence> Sequence, TWeakObjectPtr<UMovieSceneSection> Section, ScriptingKeyValueType& InDefaultValue)
+	{
+		ChannelType* Channel = ChannelHandle.Get();
+		if (Channel)
+		{
+			if (Section.IsValid())
+			{
+				Section->Modify();
+			}
+
+			using namespace UE::MovieScene;
+			SetChannelDefault(Channel, InDefaultValue);
+
+#if WITH_EDITOR
+			const FMovieSceneChannelMetaData* MetaData = ChannelHandle.GetMetaData();
+			if (MetaData && Section.IsValid() && Sequence.IsValid() && Sequence->GetMovieScene())
+			{
+				Sequence->GetMovieScene()->OnChannelChanged().Broadcast(MetaData, Section.Get());
+			}
+#endif
+			return;
+		}
+		FFrame::KismetExecutionMessage(TEXT("Invalid ChannelHandle for MovieSceneScriptingChannel, failed to set default value."), ELogVerbosity::Error);
+	}
+
+	static void RemoveDefaultFromChannel(TMovieSceneChannelHandle<ChannelType> ChannelHandle, TWeakObjectPtr<UMovieSceneSequence> Sequence, TWeakObjectPtr<UMovieSceneSection> Section)
+	{
+		ChannelType* Channel = ChannelHandle.Get();
+		if (Channel)
+		{
+			if (Section.IsValid())
+			{
+				Section->Modify();
+			}
+
+			using namespace UE::MovieScene;
+			RemoveChannelDefault(Channel);
+
+#if WITH_EDITOR
+			const FMovieSceneChannelMetaData* MetaData = ChannelHandle.GetMetaData();
+			if (MetaData && Section.IsValid() && Sequence.IsValid() && Sequence->GetMovieScene())
+			{
+				Sequence->GetMovieScene()->OnChannelChanged().Broadcast(MetaData, Section.Get());
+			}
+#endif
+			return;
+		}
+		FFrame::KismetExecutionMessage(TEXT("Invalid ChannelHandle for MovieSceneScriptingChannel, failed to remove default value."), ELogVerbosity::Error);
+	}
+
+	static TOptional<ScriptingKeyValueType> GetDefaultFromChannel(TMovieSceneChannelHandle<ChannelType> ChannelHandle)
+	{
+		ChannelType* Channel = ChannelHandle.Get();
+		if (Channel)
+		{
+			ScriptingKeyValueType Ret;
+
+			using namespace UE::MovieScene;
+			if (GetChannelDefault(Channel, Ret))
+			{
+				return TOptional<ScriptingKeyValueType>(Ret);
+			}
+
+			return TOptional<ScriptingKeyValueType>();
+		}
+
+		FFrame::KismetExecutionMessage(TEXT("Invalid ChannelHandle for MovieSceneScriptingChannel, failed to get default value."), ELogVerbosity::Error);
+		return TOptional<ScriptingKeyValueType>();
+	}
+
+	struct FTransformRetiming : UE::MovieScene::IRetimingInterface
+	{
+		FTransformRetiming(FFrameNumber InOffsetFrame, double InScale, FFrameNumber InPivotFrame, const TRange<FFrameNumber>& InRange)
+			: OffsetFrame(InOffsetFrame)
+			, Scale(InScale)
+			, PivotFrame(InPivotFrame)
+			, Range(InRange) {}
+
+		virtual FFrameTime RemapTime(FFrameTime InTime) const override
+		{
+			FFrameNumber FrameNumber = InTime.RoundToFrame();
+
+			if (Range.Contains(FrameNumber))
+			{
+				if (Scale != 0.0)
+				{
+					FrameNumber = (FrameNumber - PivotFrame) * Scale + PivotFrame;
+				}
+				return FrameNumber + OffsetFrame;
+			}
+			return InTime;
+		}
+		
+		double GetScale() const override { return 1.0; }
+		TUniquePtr<IRetimingInterface> RecurseInto(UMovieScene* InMovieScene) const override { return nullptr; }
+		void Begin(UMovieScene* InMovieScene) const override {}
+		void End(UMovieScene* InMovieScene) const override {}
+
+		FFrameNumber OffsetFrame;
+		double Scale = 0.0;
+		FFrameNumber PivotFrame;
+		TRange<FFrameNumber> Range;
+	};
+
+	static void TransformKeysInChannel(TMovieSceneChannelHandle<ChannelType> ChannelHandle, TWeakObjectPtr<UMovieSceneSequence> Sequence, TWeakObjectPtr<UMovieSceneSection> Section, FFrameNumber InOffsetFrame, double InScale, FFrameNumber InPivotFrame, FSequencerScriptingRange ScriptingRange, EMovieSceneTimeUnit TimeUnit)
+	{
+		// The Key's time is always going to be in Tick Resolution space, but the user may want to set it via Display Rate, so we convert.
+		FFrameNumber OffsetFrame = InOffsetFrame;
+		FFrameNumber PivotFrame = InPivotFrame;
+
+		UMovieSceneSequence* MovieSceneSequence = Sequence.Get();
+		FFrameRate Resolution = UMovieSceneSequenceExtensions::GetTickResolution(MovieSceneSequence);
+		TRange<FFrameNumber> Range = ScriptingRange.ToNative(Resolution);
+
+		if (Range.IsEmpty())
+		{
+			Range = TRange<FFrameNumber>::All();
+		}
+
+		if (TimeUnit == EMovieSceneTimeUnit::DisplayRate)
+		{
+			FFrameRate DisplayRate = UMovieSceneSequenceExtensions::GetDisplayRate(Sequence.Get());
+			FFrameRate TickResolution = UMovieSceneSequenceExtensions::GetTickResolution(Sequence.Get());
+
+			OffsetFrame = FFrameRate::TransformTime(OffsetFrame, DisplayRate, TickResolution).RoundToFrame();
+			PivotFrame = FFrameRate::TransformTime(PivotFrame, DisplayRate, TickResolution).RoundToFrame();
+
+			ScriptingRange.InternalRate = DisplayRate;
+			Range = ScriptingRange.ToNative(Resolution);
+		}
+
+		if (Section.IsValid())
+		{
+			Section->Modify();
+		}
+
+		ChannelType* Channel = ChannelHandle.Get();
+		if (Channel)
+		{
+			FTransformRetiming Retimer(OffsetFrame, InScale, PivotFrame, Range);
+			Channel->RemapTimes(Retimer);
+			return;
+		}
+
+		FFrame::KismetExecutionMessage(TEXT("Invalid ChannelHandle for MovieSceneScriptingChannel, failed to transform keys in channel."), ELogVerbosity::Error);
+	}
+};
+
+/** 
+* The existing Sequencer code is heavily template-based. We cannot create templated UObjects nor create UFUNCTIONS out template functions.
+* This template class serves as a way to minimize boilerplate code when creating UObject versions of the Sequencer key data.
+*/
+template<typename ChannelType, typename ChannelDataType>
+struct TMovieSceneScriptingKey
+{
+	FFrameTime GetTimeFromChannel(FKeyHandle KeyHandle, TWeakObjectPtr<UMovieSceneSequence> Sequence, EMovieSceneTimeUnit TimeUnit) const
+	{
+		if (!Sequence.IsValid())
+		{
+			FFrame::KismetExecutionMessage(TEXT("GetTime called with an invalid owning sequence."), ELogVerbosity::Error);
+			return FFrameNumber(0);
+		}
+
+		ChannelType* Channel = ChannelHandle.Get();
+		if (Channel)
+		{
+			FFrameNumber KeyTime;
+			Channel->GetKeyTime(KeyHandle, KeyTime);
+
+			// The KeyTime is always going to be in Tick Resolution space, but the user may desire it in Play Rate /w a Subframe.
+			if (TimeUnit == EMovieSceneTimeUnit::DisplayRate)
+			{
+				FFrameTime DisplayRateTime = FFrameRate::TransformTime(KeyTime, UMovieSceneSequenceExtensions::GetTickResolution(Sequence.Get()), UMovieSceneSequenceExtensions::GetDisplayRate(Sequence.Get()));
+				return DisplayRateTime;
+			}
+
+			// Tick Resolution has no sub-frame support.
+			return FFrameTime(KeyTime, 0.f);
+		}
+
+		FFrame::KismetExecutionMessage(TEXT("Invalid ChannelHandle for MovieSceneScriptingKey, failed to retrieve Time."), ELogVerbosity::Error);
+		return FFrameTime();
+	}
+
+	void SetTimeInChannel(FKeyHandle KeyHandle, TWeakObjectPtr<UMovieSceneSequence> Sequence, TWeakObjectPtr<UMovieSceneSection> Section, const FFrameNumber NewFrameNumber, EMovieSceneTimeUnit TimeUnit, float SubFrame)
+	{
+		if (!Sequence.IsValid())
+		{
+			FFrame::KismetExecutionMessage(TEXT("SetTime called with an invalid owning sequence."), ELogVerbosity::Error);
+			return;
+		}
+
+		// Clamp sub-frames to 0-1
+		SubFrame = FMath::Clamp(SubFrame, 0.f, FFrameTime::MaxSubframe);
+
+		// TickResolution doesn't support a sub-frame as you can't get finer detailed than that.
+		if (TimeUnit == EMovieSceneTimeUnit::TickResolution && SubFrame > 0.f)
+		{
+			FFrame::KismetExecutionMessage(TEXT("SetTime called with a SubFrame specified for a Tick Resolution type time! SubFrames are only allowed for Display Rate types, ignoring..."), ELogVerbosity::Error);
+			SubFrame = 0.f;
+		}
+
+		FFrameNumber KeyFrameNumber = NewFrameNumber;
+		
+		// Keys are always stored in Tick Resolution so we need to potentially convert their values.
+		if (TimeUnit == EMovieSceneTimeUnit::DisplayRate)
+		{
+			KeyFrameNumber = FFrameRate::TransformTime(FFrameTime(NewFrameNumber, SubFrame), UMovieSceneSequenceExtensions::GetDisplayRate(Sequence.Get()), UMovieSceneSequenceExtensions::GetTickResolution(Sequence.Get())).RoundToFrame();
+		}
+
+		ChannelType* Channel = ChannelHandle.Get();
+		if (Channel)
+		{
+			if (Section.IsValid())
+			{
+				Section->Modify();
+			}
+
+			Channel->SetKeyTime(KeyHandle, KeyFrameNumber);
+			return;
+		}
+
+		FFrame::KismetExecutionMessage(TEXT("Invalid ChannelHandle for MovieSceneScriptingKey, failed to set Time."), ELogVerbosity::Error);
+	}
+	
+	ChannelDataType GetValueFromChannel(FKeyHandle KeyHandle) const
+	{
+		ChannelDataType Value = ChannelDataType();
+
+		ChannelType* Channel = ChannelHandle.Get();
+		if (Channel)
+		{
+			using namespace UE::MovieScene;
+			if (!GetKeyValue(Channel, KeyHandle, Value))
+			{
+				FFrame::KismetExecutionMessage(TEXT("Invalid KeyIndex for MovieSceneScriptingKey, failed to get value. Did you forget to create the key through the channel?"), ELogVerbosity::Error);
+				return Value;
+			}
+
+			return Value;
+		}
+
+		FFrame::KismetExecutionMessage(TEXT("Invalid ChannelHandle for MovieSceneScriptingKey, failed to get value. Did you forget to create the key through the channel?"), ELogVerbosity::Error);
+		return Value;
+	}
+
+	void SetValueInChannel(FKeyHandle KeyHandle, TWeakObjectPtr<UMovieSceneSection> Section, ChannelDataType InNewValue)
+	{
+		ChannelType* Channel = ChannelHandle.Get();
+		if (Channel)
+		{
+			if (Section.IsValid())
+			{
+				Section->Modify();
+			}
+
+			using namespace UE::MovieScene;
+			if (!AssignValue(Channel, KeyHandle, InNewValue))
+			{
+				FFrame::KismetExecutionMessage(TEXT("Invalid KeyIndex for MovieSceneScriptingKey, failed to set value. Did you forget to create the key through the channel?"), ELogVerbosity::Error);
+				return;
+			}
+		}
+		else
+		{
+			FFrame::KismetExecutionMessage(TEXT("Invalid ChannelHandle for MovieSceneScriptingKey, failed to set value. Did you forget to create the key through the channel?"), ELogVerbosity::Error);
+		}
+	}
+
+public:
+	TMovieSceneChannelHandle<ChannelType> ChannelHandle;
+};
+
+#undef UE_API
diff --git a/Engine/Source/Runtime/MovieSceneTracks/Private/Sections/MovieSceneConstrainedSection.cpp b/Engine/Source/Runtime/MovieSceneTracks/Private/Sections/MovieSceneConstrainedSection.cpp
--- a/Engine/Source/Runtime/MovieSceneTracks/Private/Sections/MovieSceneConstrainedSection.cpp
+++ b/Engine/Source/Runtime/MovieSceneTracks/Private/Sections/MovieSceneConstrainedSection.cpp
@@ -19,13 +19,25 @@ UTickableConstraint* IMovieSceneConstrainedSection::FindConstraint(const FName I
 {
 	const TArray<FConstraintAndActiveChannel>& Constraints = GetConstraintsChannels();
 	const int32 ConstraintIndex = Constraints.IndexOfByPredicate([InConstraintName](const FConstraintAndActiveChannel& ConstraintChannel)
-	{
-		UTickableConstraint* Constraint = ConstraintChannel.GetConstraint().Get();
-		return Constraint && Constraint->GetFName() == InConstraintName;
-	});
+		{
+			UTickableConstraint* Constraint = ConstraintChannel.GetConstraint().Get();
+			return Constraint && Constraint->GetFName() == InConstraintName;
+		});
+	return ConstraintIndex != INDEX_NONE ? Constraints[ConstraintIndex].GetConstraint() : nullptr;
+}
+
+UTickableConstraint* IMovieSceneConstrainedSection::FindConstraint(FMovieSceneChannel* InChannel) const
+{
+	const TArray<FConstraintAndActiveChannel>& Constraints = GetConstraintsChannels();
+	const int32 ConstraintIndex = Constraints.IndexOfByPredicate([InChannel](const FConstraintAndActiveChannel& ConstraintChannel)
+		{
+			return (InChannel && &ConstraintChannel.ActiveChannel == InChannel);
+
+		});
 	return ConstraintIndex != INDEX_NONE ? Constraints[ConstraintIndex].GetConstraint() : nullptr;
 }
 
+
 bool IMovieSceneConstrainedSection::RemoveConstraint(UWorld* InWorld, const FName InConstraintName, const bool bDoNotCompensate) const
 {
 	UTickableConstraint* Constraint = FindConstraint(InConstraintName);
diff --git a/Engine/Source/Runtime/MovieSceneTracks/Public/Sections/MovieSceneConstrainedSection.h b/Engine/Source/Runtime/MovieSceneTracks/Public/Sections/MovieSceneConstrainedSection.h
--- a/Engine/Source/Runtime/MovieSceneTracks/Public/Sections/MovieSceneConstrainedSection.h
+++ b/Engine/Source/Runtime/MovieSceneTracks/Public/Sections/MovieSceneConstrainedSection.h
@@ -12,6 +12,7 @@
 #include "MovieSceneConstrainedSection.generated.h"
 
 struct FGuid;
+struct FMovieSceneChannel;
 
 /**
  * Functionality to add to sections that contain constraints
@@ -97,6 +98,11 @@ public:
 	*/
 	MOVIESCENETRACKS_API UTickableConstraint* FindConstraint(const FName InConstraintName) const;
 
+	/*
+	* Find the constraint from the channel
+	*/
+	MOVIESCENETRACKS_API UTickableConstraint* FindConstraint(FMovieSceneChannel* InChannel) const;
+
 	/*
 	* Removes the constraint named ConstraintName
 	*/
