Anim In Engine: Add transaction when doing scripted  curve filters.
#prefight
#Synced-CL 48438916

diff --git a/Engine/Plugins/MovieScene/SequencerScripting/Source/SequencerScriptingEditor/Private/SequencerCurveEditorObject.cpp b/Engine/Plugins/MovieScene/SequencerScripting/Source/SequencerScriptingEditor/Private/SequencerCurveEditorObject.cpp
--- a/Engine/Plugins/MovieScene/SequencerScripting/Source/SequencerScriptingEditor/Private/SequencerCurveEditorObject.cpp
+++ b/Engine/Plugins/MovieScene/SequencerScripting/Source/SequencerScriptingEditor/Private/SequencerCurveEditorObject.cpp
@@ -16,6 +16,10 @@
 #include "Channels/MovieSceneChannel.h"
 #include "ExtensionLibraries/MovieSceneSectionExtensions.h"
 #include "Filters/CurveEditorFilterBase.h"
+#include "ScopedTransaction.h"
+#include "Modification/Utils/ScopedCurveChange.h"
+#include "Modification/Utils/ScopedSelectionChange.h"
+#include "Modification/Keys/CurveChangeDiff.h"
 
 //For custom colors on channels, stored in editor pref's
 #include "CurveEditorSettings.h"
@@ -334,14 +338,27 @@ void USequencerCurveEditorObject::SetRandomColorForChannels(UClass* Class, const
 
 void USequencerCurveEditorObject::ApplyFilter(UCurveEditorFilterBase* Filter)
 {
+	using namespace UE::CurveEditor;
+
 	TSharedPtr<FCurveEditor> CurveEditor =  GetCurveEditor();
 
 	if (Filter && CurveEditor.IsValid())
 	{
+		const FText TransactionText = FText::Format(
+			NSLOCTEXT("CurveEditorFilterApply", "Filter Curves", "Filtered {0}|plural(one=Curve, other=Curves)"),
+			CurveEditor->GetSelection().Count()
+		);
+		const FScopedTransaction Transaction(TransactionText);
+
 		Filter->InitializeFilter(CurveEditor.ToSharedRef());
 		const TMap<FCurveModelID, FKeyHandleSet>& SelectedKeys = CurveEditor->GetSelection().GetAll();
 		TMap<FCurveModelID, FKeyHandleSet> OutKeysToSelect;
-		Filter->ApplyFilter(CurveEditor.ToSharedRef(), SelectedKeys, OutKeysToSelect);
+		{
+			const FScopedCurveChange KeyChange(FCurvesSnapshotBuilder(CurveEditor, SelectedKeys));
+
+			Filter->ApplyFilter(CurveEditor.ToSharedRef(), SelectedKeys, OutKeysToSelect);
+		}
+		const FScopedSelectionChange SelectionTransaction(CurveEditor);
 
 		// Clear their selection and then set it to the keys the filter thinks you should have selected.
 		CurveEditor->GetSelection().Clear();
