Sequencer Scripting: Fix SubFrame in AddKeyInChannel

SubFrame wasn't being used when TimeUnit is specified in DisplayRate. Note: SetTimeInChannel does convert the time correctly with SubFrame

#Synced-CL 49956477

diff --git a/Engine/Plugins/MovieScene/SequencerScripting/Source/SequencerScripting/Public/MovieSceneScriptingChannel.h b/Engine/Plugins/MovieScene/SequencerScripting/Source/SequencerScripting/Public/MovieSceneScriptingChannel.h
--- a/Engine/Plugins/MovieScene/SequencerScripting/Source/SequencerScripting/Public/MovieSceneScriptingChannel.h
+++ b/Engine/Plugins/MovieScene/SequencerScripting/Source/SequencerScripting/Public/MovieSceneScriptingChannel.h
@@ -72,11 +72,21 @@ struct TMovieSceneScriptingChannel
 		{
 			ScriptingKeyType* Key = NewObject<ScriptingKeyType>();
 
+			// Clamp sub-frames to 0-1
+			SubFrame = FMath::Clamp(SubFrame, 0.f, FFrameTime::MaxSubframe);
+
+			// TickResolution doesn't support a sub-frame as you can't get finer detailed than that.
+			if (TimeUnit == EMovieSceneTimeUnit::TickResolution && SubFrame > 0.f)
+			{
+				FFrame::KismetExecutionMessage(TEXT("SetTime called with a SubFrame specified for a Tick Resolution type time! SubFrames are only allowed for Display Rate types, ignoring..."), ELogVerbosity::Warning);
+				SubFrame = 0.f;
+			}
+
 			// The Key's time is always going to be in Tick Resolution space, but the user may want to set it via Display Rate, so we convert.
 			FFrameNumber KeyTime = InTime;
 			if (TimeUnit == EMovieSceneTimeUnit::DisplayRate)
 			{
-				KeyTime = FFrameRate::TransformTime(KeyTime, UMovieSceneSequenceExtensions::GetDisplayRate(Sequence.Get()), UMovieSceneSequenceExtensions::GetTickResolution(Sequence.Get())).RoundToFrame();
+				KeyTime = FFrameRate::TransformTime(FFrameTime(KeyTime, SubFrame), UMovieSceneSequenceExtensions::GetDisplayRate(Sequence.Get()), UMovieSceneSequenceExtensions::GetTickResolution(Sequence.Get())).RoundToFrame();
 			}
 
 			if (Section.IsValid())
@@ -454,7 +464,7 @@ struct TMovieSceneScriptingKey
 		// TickResolution doesn't support a sub-frame as you can't get finer detailed than that.
 		if (TimeUnit == EMovieSceneTimeUnit::TickResolution && SubFrame > 0.f)
 		{
-			FFrame::KismetExecutionMessage(TEXT("SetTime called with a SubFrame specified for a Tick Resolution type time! SubFrames are only allowed for Display Rate types, ignoring..."), ELogVerbosity::Error);
+			FFrame::KismetExecutionMessage(TEXT("SetTime called with a SubFrame specified for a Tick Resolution type time! SubFrames are only allowed for Display Rate types, ignoring..."), ELogVerbosity::Warning);
 			SubFrame = 0.f;
 		}
 
